Merging CL 6972148

//UE4/Dev-Enterprise/Engine/Plugins/Media/...

to //UE4/Release-4.23/Engine/Plugins/Media/...

diff --git a/Engine/Plugins/Media/TimecodeSynchronizer/Source/TimecodeSynchronizer/Public/TimecodeSynchronizer.h b/Engine/Plugins/Media/TimecodeSynchronizer/Source/TimecodeSynchronizer/Public/TimecodeSynchronizer.h
--- a/Engine/Plugins/Media/TimecodeSynchronizer/Source/TimecodeSynchronizer/Public/TimecodeSynchronizer.h
+++ b/Engine/Plugins/Media/TimecodeSynchronizer/Source/TimecodeSynchronizer/Public/TimecodeSynchronizer.h
@@ -440,7 +440,7 @@ public:
 	ETimecodeSynchronizationFrameRateSources FrameRateSource;
 
 	/** The fixed framerate to use. */
-	UPROPERTY(EditAnywhere, Category="Frame Rate Settings", Meta=(EditCondition = "IN_CPP", ClampMin="15.0"))
+	UPROPERTY(EditAnywhere, Category="Frame Rate Settings", Meta=(ClampMin="15.0"))
 	FFrameRate FixedFrameRate;
 
 public:
@@ -449,14 +449,14 @@ public:
 	ETimecodeSynchronizationTimecodeType TimecodeProviderType;
 
 	/** Custom strategy to tick in a interval. */
-	UPROPERTY(EditAnywhere, Instanced, Category="Timecode Provider", Meta=(EditCondition="IN_CPP", DisplayName="Timecode Source"))
+	UPROPERTY(EditAnywhere, Instanced, Category="Timecode Provider", Meta=(DisplayName="Timecode Source"))
 	UTimecodeProvider* TimecodeProvider;
 
 	/**
 	 * Index of the source that drives the synchronized Timecode.
 	 * The source need to be timecoded and flag as bUseForSynchronization
 	 */
-	UPROPERTY(EditAnywhere, Category="Timecode Provider", Meta=(EditCondition="IN_CPP"))
+	UPROPERTY(EditAnywhere, Category="Timecode Provider")
 	int32 MasterSynchronizationSourceIndex;
 
 public:
@@ -489,14 +489,14 @@ private:
 	TArray<UTimeSynchronizationSource*> DynamicSources;
 
 	/** What mode will be used for synchronization. */
-	UPROPERTY(EditAnywhere, Category = "Synchronization", Meta=(EditCondition="IN_CPP"))
+	UPROPERTY(EditAnywhere, Category = "Synchronization")
 	ETimecodeSynchronizationSyncMode SyncMode;
 
 	/**
 	 * When UserDefined mode is used, the number of frames delayed from the Provider's timecode.
 	 * Negative values indicate the used timecode will be ahead of the Provider's.
 	 */
-	UPROPERTY(EditAnywhere, Category = "Synchronization", Meta=(EditCondition = "IN_CPP", ClampMin="-640", ClampMax="640"))
+	UPROPERTY(EditAnywhere, Category = "Synchronization", Meta=(ClampMin="-640", ClampMax="640"))
 	int32 FrameOffset;
 
 	/**
@@ -504,11 +504,11 @@ private:
 	 * For Auto mode, this represents the number of frames behind the newest synced frame.
 	 * For AutoModeOldest, the is the of frames ahead of the last synced frame.
 	 */
-	UPROPERTY(EditAnywhere, Category = "Synchronization", Meta = (EditCondition = "IN_CPP", ClampMin = "0", ClampMax = "640"))
+	UPROPERTY(EditAnywhere, Category = "Synchronization", Meta = (ClampMin = "0", ClampMax = "640"))
 	int32 AutoFrameOffset = 3;
 
 	/** Whether or not the specified Provider's timecode rolls over. (Rollover is expected to occur at Timecode 24:00:00:00). */
-	UPROPERTY(EditAnywhere, Category = "Synchronization", Meta=(EditCondition="IN_CPP"))
+	UPROPERTY(EditAnywhere, Category = "Synchronization")
 	bool bWithRollover = false;
 
 	/** Sources used for synchronization */
