Merging CL6980351

//UE4/Dev-Enterprise/Engine/Plugins/Media/TimecodeSynchronizer/Source/TimecodeSynchronizerEditor/Private/Widgets/STimecodeSynchronizerSourceViewport.cpp

to //UE4/Release-4.23/Engine/Plugins/Media/TimecodeSynchronizer/Source/TimecodeSynchronizerEditor/Private/Widgets/STimecodeSynchronizerSourceViewport.cpp

diff --git a/Engine/Plugins/Media/TimecodeSynchronizer/Source/TimecodeSynchronizerEditor/Private/Widgets/STimecodeSynchronizerSourceViewport.cpp b/Engine/Plugins/Media/TimecodeSynchronizer/Source/TimecodeSynchronizerEditor/Private/Widgets/STimecodeSynchronizerSourceViewport.cpp
--- a/Engine/Plugins/Media/TimecodeSynchronizer/Source/TimecodeSynchronizerEditor/Private/Widgets/STimecodeSynchronizerSourceViewport.cpp
+++ b/Engine/Plugins/Media/TimecodeSynchronizer/Source/TimecodeSynchronizerEditor/Private/Widgets/STimecodeSynchronizerSourceViewport.cpp
@@ -37,6 +37,9 @@ void STimecodeSynchronizerSourceViewport::Construct(const FArguments& InArgs, UT
 	AttachedSourceIndex = InAttachedSourceIndex;
 	bIsSynchronizedSource = bInIsSynchronizedSource;
 
+	FSlateFontInfo Font18 = FCoreStyle::Get().GetFontStyle(TEXT("NormalFont"));
+	Font18.Size = 18;
+
 	ChildSlot
 	[
 		SNew(SOverlay)
@@ -97,7 +100,7 @@ void STimecodeSynchronizerSourceViewport::Construct(const FArguments& InArgs, UT
 								// Min Timecode
 								SNew(STextBlock)
 								.ColorAndOpacity(FSlateColor::UseSubduedForeground())
-								.Font(FEditorStyle::Get().GetFontStyle("FontAwesome.18"))
+								.Font(Font18)
 								.ShadowOffset(FVector2D(1.f, 1.f))
 								.Text(this, &STimecodeSynchronizerSourceViewport::HandleIntervalMinTimecodeText)
 								.Justification(ETextJustify::Right)
@@ -110,7 +113,7 @@ void STimecodeSynchronizerSourceViewport::Construct(const FArguments& InArgs, UT
 								// Max Timecode
 								SNew(STextBlock)
 								.ColorAndOpacity(FSlateColor::UseSubduedForeground())
-								.Font(FEditorStyle::Get().GetFontStyle("FontAwesome.18"))
+								.Font(Font18)
 								.ShadowOffset(FVector2D(1.f, 1.f))
 								.Text(this, &STimecodeSynchronizerSourceViewport::HandleIntervalMaxTimecodeText)
 								.Justification(ETextJustify::Right)
@@ -130,7 +133,7 @@ void STimecodeSynchronizerSourceViewport::Construct(const FArguments& InArgs, UT
 							// Display if source is the master
 							SNew(STextBlock)
 							.ColorAndOpacity(FSlateColor::UseSubduedForeground())
-							.Font(FEditorStyle::Get().GetFontStyle("FontAwesome.18"))
+							.Font(Font18)
 							.ShadowOffset(FVector2D(1.f, 1.f))
 							.Text(this, &STimecodeSynchronizerSourceViewport::HandleIsSourceMasterText)
 							.ToolTipText(LOCTEXT("OverlayMasterSourceTooltip", "Is this source used as the master"))
@@ -143,7 +146,7 @@ void STimecodeSynchronizerSourceViewport::Construct(const FArguments& InArgs, UT
 							// Current Timecode
 							SNew(STextBlock)
 							.ColorAndOpacity(FSlateColor::UseSubduedForeground())
-							.Font(FEditorStyle::Get().GetFontStyle("FontAwesome.18"))
+							.Font(Font18)
 							.ShadowOffset(FVector2D(1.f, 1.f))
 							.Text(this, &STimecodeSynchronizerSourceViewport::HandleCurrentTimecodeText)
 							.Justification(ETextJustify::Right)
