Support MediaBundle in timecode synchronizer.

Fix SMediaImage compilation issue in SMediaFrameworkVideoInputDisplay.cpp. Fix case sensitive include in SMediaImage.cpp.

diff --git a/Engine/Plugins/Media/MediaFrameworkUtilities/MediaFrameworkUtilities.uplugin b/Engine/Plugins/Media/MediaFrameworkUtilities/MediaFrameworkUtilities.uplugin
--- a/Engine/Plugins/Media/MediaFrameworkUtilities/MediaFrameworkUtilities.uplugin
+++ b/Engine/Plugins/Media/MediaFrameworkUtilities/MediaFrameworkUtilities.uplugin
@@ -45,6 +45,10 @@
 		{
 			"Name" : "OpenCVLensDistortion",
 			"Enabled" : true
+		},
+		{
+			"Name": "MediaPlayerEditor",
+			"Enabled": true
 		}
 	]
 }
diff --git a/Engine/Plugins/Media/MediaFrameworkUtilities/Source/MediaFrameworkUtilities/MediaFrameworkUtilities.Build.cs b/Engine/Plugins/Media/MediaFrameworkUtilities/Source/MediaFrameworkUtilities/MediaFrameworkUtilities.Build.cs
--- a/Engine/Plugins/Media/MediaFrameworkUtilities/Source/MediaFrameworkUtilities/MediaFrameworkUtilities.Build.cs
+++ b/Engine/Plugins/Media/MediaFrameworkUtilities/Source/MediaFrameworkUtilities/MediaFrameworkUtilities.Build.cs
@@ -28,10 +28,11 @@ namespace UnrealBuildTool.Rules
 
 			if (Target.bBuildEditor == true)
 			{
-				PrivateDependencyModuleNames.AddRange(
+                PrivateDependencyModuleNames.AddRange(
 					new string[]
 					{
-						"Slate",
+                        "MediaPlayerEditor",
+                        "Slate",
 						"SlateCore",
 						"Settings",
 						"UnrealEd"
diff --git a/Engine/Plugins/Media/MediaFrameworkUtilities/Source/MediaFrameworkUtilities/Private/MediaBundleTimeSynchronizationSource.cpp b/Engine/Plugins/Media/MediaFrameworkUtilities/Source/MediaFrameworkUtilities/Private/MediaBundleTimeSynchronizationSource.cpp
--- a/Engine/Plugins/Media/MediaFrameworkUtilities/Source/MediaFrameworkUtilities/Private/MediaBundleTimeSynchronizationSource.cpp
+++ b/Engine/Plugins/Media/MediaFrameworkUtilities/Source/MediaFrameworkUtilities/Private/MediaBundleTimeSynchronizationSource.cpp
@@ -11,7 +11,9 @@
 #include "MediaPlayerFacade.h"
 #include "MediaTexture.h"
 #include "TimeSynchronizableMediaSource.h"
-
+#if WITH_EDITOR
+#include "Widgets/SMediaImage.h"
+#endif
 
 namespace TimeSynchronizationSource
 {
@@ -250,3 +252,17 @@ FString UMediaBundleTimeSynchronizationSource::GetDisplayName() const
 {
 	return MediaBundle ? *MediaBundle->GetName() : FString();
 }
+
+#if WITH_EDITOR
+TSharedRef<SWidget> UMediaBundleTimeSynchronizationSource::GetVisualWidget() const
+{
+	if (MediaBundle && MediaBundle->GetMediaTexture())
+	{
+		return SNew(SMediaImage, MediaBundle->GetMediaTexture());
+	}
+	else
+	{
+		return SNullWidget::NullWidget;
+	}
+}
+#endif
\ No newline at end of file
diff --git a/Engine/Plugins/Media/MediaFrameworkUtilities/Source/MediaFrameworkUtilities/Public/MediaBundleTimeSynchronizationSource.h b/Engine/Plugins/Media/MediaFrameworkUtilities/Source/MediaFrameworkUtilities/Public/MediaBundleTimeSynchronizationSource.h
--- a/Engine/Plugins/Media/MediaFrameworkUtilities/Source/MediaFrameworkUtilities/Public/MediaBundleTimeSynchronizationSource.h
+++ b/Engine/Plugins/Media/MediaFrameworkUtilities/Source/MediaFrameworkUtilities/Public/MediaBundleTimeSynchronizationSource.h
@@ -39,6 +39,9 @@ public:
 	virtual void Start(const FTimeSynchronizationStartData& InStartData) override;
 	virtual void Close() override;
 	virtual FString GetDisplayName() const override;
+#if WITH_EDITOR
+	virtual TSharedRef<SWidget> GetVisualWidget() const override;
+#endif
 	//~ End UTimeSynchronizationSource Interface
 
 private:
diff --git a/Engine/Plugins/Media/MediaFrameworkUtilities/Source/MediaFrameworkUtilitiesEditor/MediaFrameworkUtilitiesEditor.Build.cs b/Engine/Plugins/Media/MediaFrameworkUtilities/Source/MediaFrameworkUtilitiesEditor/MediaFrameworkUtilitiesEditor.Build.cs
--- a/Engine/Plugins/Media/MediaFrameworkUtilities/Source/MediaFrameworkUtilitiesEditor/MediaFrameworkUtilitiesEditor.Build.cs
+++ b/Engine/Plugins/Media/MediaFrameworkUtilities/Source/MediaFrameworkUtilitiesEditor/MediaFrameworkUtilitiesEditor.Build.cs
@@ -22,7 +22,8 @@ namespace UnrealBuildTool.Rules
 					"MediaAssets",
 					"MediaFrameworkUtilities",
 					"MediaIOCore",
-					"MediaUtils",
+                    "MediaPlayerEditor",
+                    "MediaUtils",
 					"PlacementMode",
 					"PropertyEditor",
 					"Slate",
diff --git a/Engine/Plugins/Media/MediaFrameworkUtilities/Source/MediaFrameworkUtilitiesEditor/Private/CaptureTab/SMediaFrameworkCaptureOutputWidget.cpp b/Engine/Plugins/Media/MediaFrameworkUtilities/Source/MediaFrameworkUtilitiesEditor/Private/CaptureTab/SMediaFrameworkCaptureOutputWidget.cpp
--- a/Engine/Plugins/Media/MediaFrameworkUtilities/Source/MediaFrameworkUtilitiesEditor/Private/CaptureTab/SMediaFrameworkCaptureOutputWidget.cpp
+++ b/Engine/Plugins/Media/MediaFrameworkUtilities/Source/MediaFrameworkUtilitiesEditor/Private/CaptureTab/SMediaFrameworkCaptureOutputWidget.cpp
@@ -20,6 +20,7 @@
 #include "Widgets/Layout/SScaleBox.h"
 #include "Widgets/Layout/SScrollBox.h"
 #include "Widgets/Layout/SSplitter.h"
+#include "Widgets/SMediaImage.h"
 #include "Widgets/SNullWidget.h"
 #include "Widgets/SViewport.h"
 #include "Widgets/Text/STextBlock.h"
@@ -481,37 +482,7 @@ void SMediaFrameworkCaptureRenderTargetWidget::Construct(const FArguments& InArg
 	BaseArguments._CaptureOptions = InArgs._CaptureOptions;
 	SMediaFrameworkCaptureOutputWidget::Construct(BaseArguments);
 
-	// create material
-	ImageMaterial.Reset(NewObject<UMaterial>(GetTransientPackage(), NAME_None, RF_Transient));
-	UMaterialExpressionTextureSample* TextureSampler = NewObject<UMaterialExpressionTextureSample>(ImageMaterial.Get());
-	{
-		TextureSampler->Texture = RenderTarget.Get();
-		TextureSampler->AutoSetSampleType();
-	}
-
-	FExpressionOutput& Output = TextureSampler->GetOutputs()[0];
-	FExpressionInput& Input = ImageMaterial->EmissiveColor;
-	{
-		Input.Expression = TextureSampler;
-		Input.Mask = Output.Mask;
-		Input.MaskR = Output.MaskR;
-		Input.MaskG = Output.MaskG;
-		Input.MaskB = Output.MaskB;
-		Input.MaskA = Output.MaskA;
-	}
-
-	ImageMaterial->Expressions.Add(TextureSampler);
-	ImageMaterial->MaterialDomain = EMaterialDomain::MD_UI;
-	ImageMaterial->PostEditChange();
-
-	// create Slate brush
-	ImageMaterialBrush = MakeShareable(new FSlateBrush());
-	{
-		ImageMaterialBrush->SetResourceObject(ImageMaterial.Get());
-	}
-
-	TSharedRef<SWidget> PictureBox = SNew(SImage)
-		.Image(ImageMaterialBrush.IsValid() ? ImageMaterialBrush.Get() : FEditorStyle::GetBrush("WhiteTexture"));
+	TSharedRef<SWidget> PictureBox = SNew(SMediaImage, RenderTarget.Get());
 
 	this->ChildSlot
 	[
diff --git a/Engine/Plugins/Media/MediaFrameworkUtilities/Source/MediaFrameworkUtilitiesEditor/Private/VideoInputTab/SMediaFrameworkVideoInputDisplay.cpp b/Engine/Plugins/Media/MediaFrameworkUtilities/Source/MediaFrameworkUtilitiesEditor/Private/VideoInputTab/SMediaFrameworkVideoInputDisplay.cpp
--- a/Engine/Plugins/Media/MediaFrameworkUtilities/Source/MediaFrameworkUtilitiesEditor/Private/VideoInputTab/SMediaFrameworkVideoInputDisplay.cpp
+++ b/Engine/Plugins/Media/MediaFrameworkUtilities/Source/MediaFrameworkUtilitiesEditor/Private/VideoInputTab/SMediaFrameworkVideoInputDisplay.cpp
@@ -29,6 +29,7 @@
 #include "Widgets/Layout/SBox.h"
 #include "Widgets/Notifications/SNotificationList.h"
 #include "Widgets/SBoxPanel.h"
+#include "Widgets/SMediaImage.h"
 #include "Widgets/SWidget.h"
 #include "Widgets/Text/STextBlock.h"
 
@@ -235,46 +236,6 @@ void SMediaFrameworkVideoInputDisplay::OnMediaClosed()
 	}
 }
 
-const FSlateBrush* SMediaFrameworkVideoInputDisplay::CreateDisplayBrush(UMediaTexture* InTexture)
-{
-	if (InTexture)
-	{
-		// create wrapper material
-		Material = NewObject<UMaterial>(GetTransientPackage(), NAME_None, RF_Transient | RF_Public);
-		UMaterialExpressionTextureSample* TextureSampler = NewObject<UMaterialExpressionTextureSample>(Material);
-		{
-			TextureSampler->Texture = InTexture;
-			TextureSampler->AutoSetSampleType();
-		}
-
-		FExpressionOutput& Output = TextureSampler->GetOutputs()[0];
-		FExpressionInput& Input = Material->EmissiveColor;
-		{
-			Input.Expression = TextureSampler;
-			Input.Mask = Output.Mask;
-			Input.MaskR = Output.MaskR;
-			Input.MaskG = Output.MaskG;
-			Input.MaskB = Output.MaskB;
-			Input.MaskA = Output.MaskA;
-		}
-
-		Material->Expressions.Add(TextureSampler);
-		Material->MaterialDomain = EMaterialDomain::MD_UI;
-		Material->PostEditChange();
-		Material->AddToRoot();
-
-		// create Slate brush
-		MaterialBrush = MakeShareable(new FSlateBrush());
-		{
-			MaterialBrush->SetResourceObject(Material);
-		}
-
-		return MaterialBrush.Get();
-	}
-
-	return FEditorStyle::GetDefaultBrush();
-}
-
 TSharedRef<SWidget> SMediaFrameworkVideoInputDisplay::ConstructVideoDisplay()
 {
 	UMediaTexture* MediaTexture = GetMediaTexture();
@@ -288,8 +249,7 @@ TSharedRef<SWidget> SMediaFrameworkVideoInputDisplay::ConstructVideoDisplay()
 				return VideoInputDisplayUtilities::CalculateWitdhForDisplay(GetMediaTexture());
 			})
 			[
-				SNew(SImage)
-				.Image(CreateDisplayBrush(MediaTexture))
+				SNew(SMediaImage, MediaTexture)
 			];
 	}
 
diff --git a/Engine/Plugins/Media/MediaFrameworkUtilities/Source/MediaFrameworkUtilitiesEditor/Private/VideoInputTab/SMediaFrameworkVideoInputDisplay.h b/Engine/Plugins/Media/MediaFrameworkUtilities/Source/MediaFrameworkUtilitiesEditor/Private/VideoInputTab/SMediaFrameworkVideoInputDisplay.h
--- a/Engine/Plugins/Media/MediaFrameworkUtilities/Source/MediaFrameworkUtilitiesEditor/Private/VideoInputTab/SMediaFrameworkVideoInputDisplay.h
+++ b/Engine/Plugins/Media/MediaFrameworkUtilities/Source/MediaFrameworkUtilitiesEditor/Private/VideoInputTab/SMediaFrameworkVideoInputDisplay.h
@@ -59,7 +59,6 @@ protected:
 	friend UMediaFrameworkVideoInputDisplayCallback;
 
 private:
-	const FSlateBrush* CreateDisplayBrush(UMediaTexture* InTexture);
 	TSharedRef<SWidget> ConstructVideoDisplay();
 	TSharedRef<SWidget> ConstructVideoStateDisplay(const FString& SourceName);
 
diff --git a/Engine/Plugins/Media/MediaPlayerEditor/Source/MediaPlayerEditor/Private/Widgets/SMediaImage.cpp b/Engine/Plugins/Media/MediaPlayerEditor/Source/MediaPlayerEditor/Private/Widgets/SMediaImage.cpp
new file mode 100644
--- /dev/null
+++ b/Engine/Plugins/Media/MediaPlayerEditor/Source/MediaPlayerEditor/Private/Widgets/SMediaImage.cpp
@@ -0,0 +1,82 @@
+// Copyright 1998-2019 Epic Games, Inc. All Rights Reserved.
+
+#include "Widgets/SMediaImage.h"
+
+#include "EditorStyleSet.h"
+#include "Materials/Material.h"
+#include "Styling/SlateBrush.h"
+#include "Widgets/Layout/SScaleBox.h"
+#include "Widgets/Images/SImage.h"
+#include "Widgets/Input/SButton.h"
+
+
+SMediaImage::SMediaImage()
+	: Collector(this)
+	, Material(nullptr)
+	, MaterialBrush(nullptr)
+	, TextureSampler(nullptr)
+{ }
+
+void SMediaImage::Construct(const FArguments& InArgs, UTexture* InTexture)
+{
+	// The Slate brush that renders the material.
+	BrushImageSize = InArgs._BrushImageSize;
+
+	if (InTexture != nullptr)
+	{
+		// create wrapper material
+		Material = NewObject<UMaterial>(GetTransientPackage(), NAME_None, RF_Transient);
+
+		if (Material != nullptr)
+		{
+			TextureSampler = NewObject<UMaterialExpressionTextureSample>(Material);
+			{
+				TextureSampler->Texture = InTexture;
+				TextureSampler->AutoSetSampleType();
+			}
+
+			FExpressionOutput& Output = TextureSampler->GetOutputs()[0];
+			FExpressionInput& Input = Material->EmissiveColor;
+			{
+				Input.Expression = TextureSampler;
+				Input.Mask = Output.Mask;
+				Input.MaskR = Output.MaskR;
+				Input.MaskG = Output.MaskG;
+				Input.MaskB = Output.MaskB;
+				Input.MaskA = Output.MaskA;
+			}
+
+			Material->Expressions.Add(TextureSampler);
+			Material->MaterialDomain = EMaterialDomain::MD_UI;
+			Material->PostEditChange();
+		}
+
+		// create Slate brush
+		MaterialBrush = MakeShareable(new FSlateBrush());
+		{
+			MaterialBrush->SetResourceObject(Material);
+		}
+	}
+
+	ChildSlot
+	[
+		SNew(SScaleBox)
+		.Stretch_Lambda([]() -> EStretch::Type { return EStretch::Fill;	})
+		[
+			SNew(SImage)
+			.Image(MaterialBrush.IsValid() ? MaterialBrush.Get() : FEditorStyle::GetBrush("WhiteTexture"))
+		]
+	];
+}
+
+void SMediaImage::Tick(const FGeometry& InAllottedGeometry, const double InCurrentTime, const float InDeltaTime)
+{
+	SCompoundWidget::Tick(InAllottedGeometry, InCurrentTime, InDeltaTime);
+
+	if (BrushImageSize.IsSet())
+	{
+		FVector2D Size = BrushImageSize.Get();
+		MaterialBrush->ImageSize.X = Size.X;
+		MaterialBrush->ImageSize.Y = Size.Y;
+	}
+}
diff --git a/Engine/Plugins/Media/MediaPlayerEditor/Source/MediaPlayerEditor/Private/Widgets/SMediaPlayerEditorOutput.cpp b/Engine/Plugins/Media/MediaPlayerEditor/Source/MediaPlayerEditor/Private/Widgets/SMediaPlayerEditorOutput.cpp
--- a/Engine/Plugins/Media/MediaPlayerEditor/Source/MediaPlayerEditor/Private/Widgets/SMediaPlayerEditorOutput.cpp
+++ b/Engine/Plugins/Media/MediaPlayerEditor/Source/MediaPlayerEditor/Private/Widgets/SMediaPlayerEditorOutput.cpp
@@ -13,14 +13,14 @@
 #include "MediaTexture.h"
 #include "Styling/SlateBrush.h"
 #include "Widgets/Images/SImage.h"
+#include "Widgets/SMediaImage.h"
 
 
 /* SMediaPlayerEditorOutput structors
  *****************************************************************************/
 
 SMediaPlayerEditorOutput::SMediaPlayerEditorOutput()
-	: Material(nullptr)
-	, MediaPlayer(nullptr)
+	: MediaPlayer(nullptr)
 	, MediaTexture(nullptr)
 	, SoundComponent(nullptr)
 { }
@@ -34,12 +34,6 @@ SMediaPlayerEditorOutput::~SMediaPlayerEditorOutput()
 		MediaPlayer.Reset();
 	}
 
-	if (Material != nullptr)
-	{
-		Material->RemoveFromRoot();
-		Material = nullptr;
-	}
-
 	if (MediaTexture != nullptr)
 	{
 		MediaTexture->RemoveFromRoot();
@@ -88,44 +82,18 @@ void SMediaPlayerEditorOutput::Construct(const FArguments& InArgs, UMediaPlayer&
 		MediaTexture->AddToRoot();
 	}
 
-	// create wrapper material
-	Material = NewObject<UMaterial>(GetTransientPackage(), NAME_None, RF_Transient | RF_Public);
-
-	if (Material != nullptr)
-	{
-		TextureSampler = NewObject<UMaterialExpressionTextureSample>(Material);
-		{
-			TextureSampler->Texture = MediaTexture;
-			TextureSampler->AutoSetSampleType();
-		}
-
-		FExpressionOutput& Output = TextureSampler->GetOutputs()[0];
-		FExpressionInput& Input = Material->EmissiveColor;
-		{
-			Input.Expression = TextureSampler;
-			Input.Mask = Output.Mask;
-			Input.MaskR = Output.MaskR;
-			Input.MaskG = Output.MaskG;
-			Input.MaskB = Output.MaskB;
-			Input.MaskA = Output.MaskA;
-		}
-
-		Material->Expressions.Add(TextureSampler);
-		Material->MaterialDomain = EMaterialDomain::MD_UI;
-		Material->PostEditChange();
-		Material->AddToRoot();
-	}
-
-	// create Slate brush
-	MaterialBrush = MakeShareable(new FSlateBrush());
-	{
-		MaterialBrush->SetResourceObject(Material);
-	}
+	TSharedRef<SMediaImage> MediaImage = SNew(SMediaImage, MediaTexture)
+		.BrushImageSize_Lambda([&]()
+			{
+				if (MediaTexture)
+					return FVector2D(MediaTexture->GetSurfaceWidth(), MediaTexture->GetSurfaceHeight());
+				else
+					return FVector2D::ZeroVector;
+			});
 
 	ChildSlot
 	[
-		SNew(SImage)
-			.Image(MaterialBrush.Get())
+		MediaImage
 	];
 
 	MediaPlayer->OnMediaEvent().AddRaw(this, &SMediaPlayerEditorOutput::HandleMediaPlayerMediaEvent);
@@ -137,16 +105,6 @@ void SMediaPlayerEditorOutput::Construct(const FArguments& InArgs, UMediaPlayer&
 
 void SMediaPlayerEditorOutput::Tick(const FGeometry& AllottedGeometry, const double InCurrentTime, const float InDeltaTime)
 {
-	if (MediaTexture != nullptr)
-	{
-		MaterialBrush->ImageSize.X = MediaTexture->GetSurfaceWidth();
-		MaterialBrush->ImageSize.Y = MediaTexture->GetSurfaceHeight();
-	}
-	else
-	{
-		MaterialBrush->ImageSize = FVector2D::ZeroVector;
-	}
-
 	if (SoundComponent != nullptr)
 	{
 		SoundComponent->UpdatePlayer();
diff --git a/Engine/Plugins/Media/MediaPlayerEditor/Source/MediaPlayerEditor/Private/Widgets/SMediaPlayerEditorOutput.h b/Engine/Plugins/Media/MediaPlayerEditor/Source/MediaPlayerEditor/Private/Widgets/SMediaPlayerEditorOutput.h
--- a/Engine/Plugins/Media/MediaPlayerEditor/Source/MediaPlayerEditor/Private/Widgets/SMediaPlayerEditorOutput.h
+++ b/Engine/Plugins/Media/MediaPlayerEditor/Source/MediaPlayerEditor/Private/Widgets/SMediaPlayerEditorOutput.h
@@ -7,14 +7,10 @@
 #include "Widgets/DeclarativeSyntaxSupport.h"
 #include "Widgets/SCompoundWidget.h"
 
-class UMaterial;
-class UMaterialExpressionTextureSample;
 class UMediaPlayer;
 class UMediaSoundComponent;
 class UMediaTexture;
 
-struct FSlateBrush;
-
 enum class EMediaEvent;
 
 
@@ -60,12 +56,6 @@ private:
 
 private:
 
-	/** The material that wraps the video texture for display in an SImage. */
-	UMaterial* Material;
-
-	/** The Slate brush that renders the material. */
-	TSharedPtr<FSlateBrush> MaterialBrush;
-
 	/** The media player whose video texture is shown in this widget. */
 	TWeakObjectPtr<UMediaPlayer> MediaPlayer;
 
@@ -74,7 +64,4 @@ private:
 
 	/** The sound component to play the media player's audio output. */
 	UMediaSoundComponent* SoundComponent;
-
-	/** The video texture sampler in the wrapper material. */
-	UMaterialExpressionTextureSample* TextureSampler;
 };
diff --git a/Engine/Plugins/Media/MediaPlayerEditor/Source/MediaPlayerEditor/Public/Widgets/SMediaImage.h b/Engine/Plugins/Media/MediaPlayerEditor/Source/MediaPlayerEditor/Public/Widgets/SMediaImage.h
new file mode 100644
--- /dev/null
+++ b/Engine/Plugins/Media/MediaPlayerEditor/Source/MediaPlayerEditor/Public/Widgets/SMediaImage.h
@@ -0,0 +1,82 @@
+// Copyright 1998-2019 Epic Games, Inc. All Rights Reserved.
+
+#pragma once
+
+#include "UObject/GCObject.h"
+
+#include "Widgets//DeclarativeSyntaxSupport.h"
+#include "Widgets/SCompoundWidget.h"
+
+class UMaterial;
+class UMaterialExpressionTextureSample;
+class UTexture;
+
+class MEDIAPLAYEREDITOR_API SMediaImage : public SCompoundWidget
+{
+public:
+	SLATE_BEGIN_ARGS(SMediaImage) { }
+	SLATE_ATTRIBUTE(FVector2D, BrushImageSize);
+	SLATE_END_ARGS()
+
+public:
+
+	/** Default constructor. */
+	SMediaImage();
+
+public:
+
+	/**
+	 * Construct this widget
+	 *
+	 * @param InArgs The declaration data for this widget.
+	 * @param InTexture The source texture to render if any
+	 */
+	void Construct(const FArguments& InArgs, UTexture* InTexture);
+
+	/**
+	 * Tick this widget
+	 * 
+	 * @param InAllottedGeometry Geometry of the widget.
+	 * @param InCurrentTime CurrentTime of the engine.
+	 * @param InDeltaTime Deltatime since last frame.
+	 */
+	virtual void Tick(const FGeometry& InAllottedGeometry, const double InCurrentTime, const float InDeltaTime) override;
+
+private:
+
+	class FInternalReferenceCollector : public FGCObject
+	{
+	public:
+		FInternalReferenceCollector(SMediaImage* InObject)
+			: Object(InObject)
+		{
+		}
+
+		//~ FGCObject interface
+		virtual void AddReferencedObjects(FReferenceCollector& InCollector) override
+		{
+			InCollector.AddReferencedObject(Object->Material);
+			InCollector.AddReferencedObject(Object->TextureSampler);
+		}
+
+	private:
+		SMediaImage* Object;
+	};
+	friend FInternalReferenceCollector;
+
+	/** Collector to keep the UObject alive. */
+	FInternalReferenceCollector Collector;
+
+	/** The material that wraps the video texture for display in an SImage. */
+	UMaterial* Material;
+
+	/** The Slate brush that renders the material. */
+	TSharedPtr<FSlateBrush> MaterialBrush;
+
+	/** The video texture sampler in the wrapper material. */
+	UMaterialExpressionTextureSample* TextureSampler;
+
+	/** Brush image size attribute. */
+	TAttribute<FVector2D> BrushImageSize;
+};
+
diff --git a/Engine/Plugins/Media/TimecodeSynchronizer/Source/TimecodeSynchronizer/Private/MediaPlayerTimeSynchronizationSource.cpp b/Engine/Plugins/Media/TimecodeSynchronizer/Source/TimecodeSynchronizer/Private/MediaPlayerTimeSynchronizationSource.cpp
--- a/Engine/Plugins/Media/TimecodeSynchronizer/Source/TimecodeSynchronizer/Private/MediaPlayerTimeSynchronizationSource.cpp
+++ b/Engine/Plugins/Media/TimecodeSynchronizer/Source/TimecodeSynchronizer/Private/MediaPlayerTimeSynchronizationSource.cpp
@@ -10,6 +10,10 @@
 #include "MediaTexture.h"
 #include "TimeSynchronizableMediaSource.h"
 
+#if WITH_EDITOR
+#include "Widgets/SMediaImage.h"
+#endif
+
 
 #if WITH_EDITOR
 void UMediaPlayerTimeSynchronizationSource::PostEditChangeProperty(FPropertyChangedEvent& PropertyChangedEvent)
@@ -236,3 +240,16 @@ FString UMediaPlayerTimeSynchronizationSource::GetDisplayName() const
 	return MediaTexture && MediaTexture->GetMediaPlayer() ? *MediaTexture->GetMediaPlayer()->GetName() : FString();
 }
 
+#if WITH_EDITOR
+TSharedRef<SWidget> UMediaPlayerTimeSynchronizationSource::GetVisualWidget() const
+{
+	if (MediaTexture)
+	{
+		return SNew(SMediaImage, MediaTexture);
+	}
+	else
+	{
+		return SNullWidget::NullWidget;
+	}
+}
+#endif
\ No newline at end of file
diff --git a/Engine/Plugins/Media/TimecodeSynchronizer/Source/TimecodeSynchronizer/Public/MediaPlayerTimeSynchronizationSource.h b/Engine/Plugins/Media/TimecodeSynchronizer/Source/TimecodeSynchronizer/Public/MediaPlayerTimeSynchronizationSource.h
--- a/Engine/Plugins/Media/TimecodeSynchronizer/Source/TimecodeSynchronizer/Public/MediaPlayerTimeSynchronizationSource.h
+++ b/Engine/Plugins/Media/TimecodeSynchronizer/Source/TimecodeSynchronizer/Public/MediaPlayerTimeSynchronizationSource.h
@@ -44,6 +44,9 @@ public:
 	virtual void Start(const FTimeSynchronizationStartData& InStartData) override;
 	virtual void Close() override;
 	virtual FString GetDisplayName() const override;
+#if WITH_EDITOR
+	virtual TSharedRef<SWidget> GetVisualWidget() const override;
+#endif
 	//~ End UTimeSynchronizationSource Interface
 
 private:
diff --git a/Engine/Plugins/Media/TimecodeSynchronizer/Source/TimecodeSynchronizer/TimecodeSynchronizer.Build.cs b/Engine/Plugins/Media/TimecodeSynchronizer/Source/TimecodeSynchronizer/TimecodeSynchronizer.Build.cs
--- a/Engine/Plugins/Media/TimecodeSynchronizer/Source/TimecodeSynchronizer/TimecodeSynchronizer.Build.cs
+++ b/Engine/Plugins/Media/TimecodeSynchronizer/Source/TimecodeSynchronizer/TimecodeSynchronizer.Build.cs
@@ -4,17 +4,26 @@ using UnrealBuildTool;
 
 public class TimecodeSynchronizer : ModuleRules
 {
-	public TimecodeSynchronizer(ReadOnlyTargetRules Target) : base(Target)
-	{
-		PublicDependencyModuleNames.AddRange(
-			new string[] {
+    public TimecodeSynchronizer(ReadOnlyTargetRules Target) : base(Target)
+    {
+        PublicDependencyModuleNames.AddRange(
+            new string[] {
                 "Core",
-				"CoreUObject",
-				"Engine",
-				"Media",
-				"MediaAssets",
-				"MediaUtils",
+                "CoreUObject",
+                "Engine",
+                "Media",
+                "MediaAssets",
+                "MediaUtils",
                 "TimeManagement",
             });
-	}
+
+        if (Target.bBuildEditor == true)
+        {
+            PrivateDependencyModuleNames.AddRange(
+                new string[] {
+                "MediaPlayerEditor",
+                "SlateCore",
+                });
+        }
+    }
 }
diff --git a/Engine/Plugins/Media/TimecodeSynchronizer/Source/TimecodeSynchronizerEditor/Private/Widgets/STimecodeSynchronizerSourceViewer.cpp b/Engine/Plugins/Media/TimecodeSynchronizer/Source/TimecodeSynchronizerEditor/Private/Widgets/STimecodeSynchronizerSourceViewer.cpp
--- a/Engine/Plugins/Media/TimecodeSynchronizer/Source/TimecodeSynchronizerEditor/Private/Widgets/STimecodeSynchronizerSourceViewer.cpp
+++ b/Engine/Plugins/Media/TimecodeSynchronizer/Source/TimecodeSynchronizerEditor/Private/Widgets/STimecodeSynchronizerSourceViewer.cpp
@@ -3,17 +3,15 @@
 #include "Widgets/STimecodeSynchronizerSourceViewer.h"
 
 
-
 #include "Framework/MultiBox/MultiBoxBuilder.h"
 #include "Framework/Notifications/NotificationManager.h"
 #include "MediaPlayer.h"
 #include "MediaPlayerTimeSynchronizationSource.h"
 #include "MediaTexture.h"
 #include "Modules/ModuleManager.h"
+#include "Widgets/Layout/SBorder.h"
 #include "Widgets/SBoxPanel.h"
 #include "Widgets/SOverlay.h"
-#include "Widgets/Layout/SBorder.h"
-#include "Widgets/Notifications/SNotificationList.h"
 #include "Widgets/STimecodeSynchronizerSourceViewport.h"
 
 
@@ -60,6 +58,18 @@ void STimecodeSynchronizerSourceViewer::Construct(const FArguments& InArgs, UTim
 	PopulateActiveSources();
 }
 
+TSharedRef<SWidget> STimecodeSynchronizerSourceViewer::GetVisualWidget(const FTimecodeSynchronizerActiveTimecodedInputSource& InSource)
+{
+	if (const UTimeSynchronizationSource* SyncSource = InSource.GetInputSource())
+	{
+		return SyncSource->GetVisualWidget();
+	}
+	else
+	{
+		return SNullWidget::NullWidget;
+	}
+}
+
 void STimecodeSynchronizerSourceViewer::PopulateActiveSources()
 {
 	ViewportVerticalBox->ClearChildren();
@@ -74,24 +84,19 @@ void STimecodeSynchronizerSourceViewer::PopulateActiveSources()
 				for (int32 Index = 0; Index < TimecodedSources.Num(); ++Index)
 				{
 					const FTimecodeSynchronizerActiveTimecodedInputSource& Source = TimecodedSources[Index];
-					if (const UTimeSynchronizationSource* SyncSource = Source.GetInputSource())
-					{
-						const UMediaPlayerTimeSynchronizationSource* MediaPlayerSource = Cast<UMediaPlayerTimeSynchronizationSource>(SyncSource);
-						const UMediaTexture* TextureArg = MediaPlayerSource ? MediaPlayerSource->MediaTexture : nullptr;
-
-						//Add a Viewport Widget for each active Source
-						Owner->AddSlot()
-							.Padding(1.0f, 1.0f, 1.0f, 1.0f)
+
+					//Add a Viewport Widget for each active Source
+					Owner->AddSlot()
+						.Padding(1.0f, 1.0f, 1.0f, 1.0f)
+						[
+							SNew(SBorder)
+							.BorderImage(FCoreStyle::Get().GetBrush("GreenBrush"))
+							.Padding(0.0f)
 							[
-								SNew(SBorder)
-								.BorderImage(FCoreStyle::Get().GetBrush("GreenBrush"))
-								.Padding(0.0f)
-								[
-									//Source area
-									SNew(STimecodeSynchronizerSourceViewport, Synchronizer, Index, bSynchronizedSources, const_cast<UMediaTexture*>(TextureArg))
-								]
-							];
-					}
+								//Source area
+								SNew(STimecodeSynchronizerSourceViewport, Synchronizer, Index, bSynchronizedSources, GetVisualWidget(TimecodedSources[Index]))
+							]
+						];
 				}
 			}
 		};
@@ -103,7 +108,7 @@ void STimecodeSynchronizerSourceViewer::PopulateActiveSources()
 
 void STimecodeSynchronizerSourceViewer::HandleSynchronizationEvent(ETimecodeSynchronizationEvent Event)
 {
-	if (Event == ETimecodeSynchronizationEvent::SynchronizationStarted)
+	if (Event == ETimecodeSynchronizationEvent::SynchronizationSucceeded || Event == ETimecodeSynchronizationEvent::SynchronizationStopped)
 	{
 		PopulateActiveSources();
 	}
diff --git a/Engine/Plugins/Media/TimecodeSynchronizer/Source/TimecodeSynchronizerEditor/Private/Widgets/STimecodeSynchronizerSourceViewer.h b/Engine/Plugins/Media/TimecodeSynchronizer/Source/TimecodeSynchronizerEditor/Private/Widgets/STimecodeSynchronizerSourceViewer.h
--- a/Engine/Plugins/Media/TimecodeSynchronizer/Source/TimecodeSynchronizerEditor/Private/Widgets/STimecodeSynchronizerSourceViewer.h
+++ b/Engine/Plugins/Media/TimecodeSynchronizer/Source/TimecodeSynchronizerEditor/Private/Widgets/STimecodeSynchronizerSourceViewer.h
@@ -51,7 +51,7 @@ private:
 
 	void PopulateActiveSources();
 	void HandleSynchronizationEvent(ETimecodeSynchronizationEvent Event);
-
+	static TSharedRef<SWidget> GetVisualWidget(const FTimecodeSynchronizerActiveTimecodedInputSource& InSource);
 
 private:
 
diff --git a/Engine/Plugins/Media/TimecodeSynchronizer/Source/TimecodeSynchronizerEditor/Private/Widgets/STimecodeSynchronizerSourceViewport.cpp b/Engine/Plugins/Media/TimecodeSynchronizer/Source/TimecodeSynchronizerEditor/Private/Widgets/STimecodeSynchronizerSourceViewport.cpp
--- a/Engine/Plugins/Media/TimecodeSynchronizer/Source/TimecodeSynchronizerEditor/Private/Widgets/STimecodeSynchronizerSourceViewport.cpp
+++ b/Engine/Plugins/Media/TimecodeSynchronizer/Source/TimecodeSynchronizerEditor/Private/Widgets/STimecodeSynchronizerSourceViewport.cpp
@@ -6,13 +6,11 @@
 #include "EditorStyleSet.h"
 #include "Engine/Texture.h"
 #include "Materials/Material.h"
-#include "Materials/MaterialExpressionTextureSample.h"
 #include "Misc/App.h"
 #include "Widgets/SBoxPanel.h"
 #include "Widgets/Input/SEditableTextBox.h"
 #include "Widgets/SOverlay.h"
 #include "Widgets/Images/SImage.h"
-#include "Widgets/Layout/SScaleBox.h"
 #include "Widgets/Text/STextBlock.h"
 
 
@@ -24,63 +22,20 @@
 *****************************************************************************/
 
 STimecodeSynchronizerSourceViewport::STimecodeSynchronizerSourceViewport()
-	: Collector(this)
-	, SourceTextBox(nullptr)
+	: SourceTextBox(nullptr)
 	, TimecodeSynchronization(nullptr)
 	, AttachedSourceIndex(INDEX_NONE)
 	, bIsSynchronizedSource(false)
-	, Material(nullptr)
-	, MaterialBrush(nullptr)
-	, TextureSampler(nullptr)
 { }
 
 /* STimecodeSynchronizerSourceViewport interface
 *****************************************************************************/
 
-void STimecodeSynchronizerSourceViewport::Construct(const FArguments& InArgs, UTimecodeSynchronizer* InTimecodeSynchronizer, int32 InAttachedSourceIndex, bool bInIsSynchronizedSource, UTexture* InTexture)
+void STimecodeSynchronizerSourceViewport::Construct(const FArguments& InArgs, UTimecodeSynchronizer* InTimecodeSynchronizer, int32 InAttachedSourceIndex, bool bInIsSynchronizedSource, TSharedRef<SWidget> InVisualWidget)
 {
-	TimecodeSynchronization = InTimecodeSynchronizer;
+	TimecodeSynchronization.Reset(InTimecodeSynchronizer);
 	AttachedSourceIndex = InAttachedSourceIndex;
 	bIsSynchronizedSource = bInIsSynchronizedSource;
-	
-	//If no texture is attached, a default static texture will be used (for non live viewable sources)
-	MaterialBrush = nullptr;
-	if (InTexture != nullptr)
-	{
-		// create wrapper material
-		Material = NewObject<UMaterial>(GetTransientPackage(), NAME_None, RF_Transient);
-
-		if (Material != nullptr)
-		{
-			TextureSampler = NewObject<UMaterialExpressionTextureSample>(Material);
-			{
-				TextureSampler->Texture = InTexture;
-				TextureSampler->AutoSetSampleType();
-			}
-
-			FExpressionOutput& Output = TextureSampler->GetOutputs()[0];
-			FExpressionInput& Input = Material->EmissiveColor;
-			{
-				Input.Expression = TextureSampler;
-				Input.Mask = Output.Mask;
-				Input.MaskR = Output.MaskR;
-				Input.MaskG = Output.MaskG;
-				Input.MaskB = Output.MaskB;
-				Input.MaskA = Output.MaskA;
-			}
-
-			Material->Expressions.Add(TextureSampler);
-			Material->MaterialDomain = EMaterialDomain::MD_UI;
-			Material->PostEditChange();
-		}
-		
-		// create Slate brush
-		MaterialBrush = MakeShareable(new FSlateBrush());
-		{
-			MaterialBrush->SetResourceObject(Material);
-		}
-	}
-
 
 	ChildSlot
 	[
@@ -119,13 +74,7 @@ void STimecodeSynchronizerSourceViewport::Construct(const FArguments& InArgs, UT
 				SNew(SOverlay)
 				+ SOverlay::Slot()
 				[
-					// Live view of the Source
-					SNew(SScaleBox)
-					.Stretch_Lambda([]() -> EStretch::Type { return EStretch::Fill;	})
-					[
-						SNew(SImage)
-						.Image(MaterialBrush.IsValid() ? MaterialBrush.Get() : FEditorStyle::GetBrush("WhiteTexture"))
-					]
+					InVisualWidget
 				]
 
 				+ SOverlay::Slot()
diff --git a/Engine/Plugins/Media/TimecodeSynchronizer/Source/TimecodeSynchronizerEditor/Private/Widgets/STimecodeSynchronizerSourceViewport.h b/Engine/Plugins/Media/TimecodeSynchronizer/Source/TimecodeSynchronizerEditor/Private/Widgets/STimecodeSynchronizerSourceViewport.h
--- a/Engine/Plugins/Media/TimecodeSynchronizer/Source/TimecodeSynchronizerEditor/Private/Widgets/STimecodeSynchronizerSourceViewport.h
+++ b/Engine/Plugins/Media/TimecodeSynchronizer/Source/TimecodeSynchronizerEditor/Private/Widgets/STimecodeSynchronizerSourceViewport.h
@@ -5,16 +5,12 @@
 #include "Styling/ISlateStyle.h"
 #include "Widgets/DeclarativeSyntaxSupport.h"
 #include "Widgets/SCompoundWidget.h"
-#include "UObject/GCObject.h"
-
+#include "UObject/StrongObjectPtr.h"
 
 class SEditableTextBox;
 struct FSlateBrush;
 struct FTimecodeSynchronizerActiveTimecodedInputSource;
-class UMaterial;
 class UTimecodeSynchronizer;
-class UTexture;
-class UMaterialExpressionTextureSample;
 
 class STimecodeSynchronizerSourceViewport
 	: public SCompoundWidget
@@ -39,7 +35,7 @@ public:
 	 * @param InAttachedSource The Source to display.
 	 * @param InTexture The source texture to render if any
 	 */
-	void Construct(const FArguments& InArgs, UTimecodeSynchronizer* InTimecodeSynchronization, int32 InAttachedSourceIndex, bool InTimecodedSource, UTexture* InTexture);
+	void Construct(const FArguments& InArgs, UTimecodeSynchronizer* InTimecodeSynchronization, int32 InAttachedSourceIndex, bool InTimecodedSource, TSharedRef<SWidget> InVisualWidget);
 
 private:
 
@@ -59,48 +55,16 @@ private:
 	const FTimecodeSynchronizerActiveTimecodedInputSource* GetAttachedSource() const;
 
 private:
-	class FInternalReferenceCollector : public FGCObject
-	{
-	public:
-		FInternalReferenceCollector(STimecodeSynchronizerSourceViewport* InObject)
-			: Object(InObject)
-		{
-		}
-
-		//~ FGCObject interface
-		virtual void AddReferencedObjects(FReferenceCollector& InCollector) override
-		{
-			InCollector.AddReferencedObject(Object->TimecodeSynchronization);
-			InCollector.AddReferencedObject(Object->Material);
-			InCollector.AddReferencedObject(Object->TextureSampler);
-		}
-
-	private:
-		STimecodeSynchronizerSourceViewport* Object;
-	};
-	friend FInternalReferenceCollector;
-
-	/** Collector to keep the UObject alive. */
-	FInternalReferenceCollector Collector;
 
 	/** Media Source name text box. */
 	TSharedPtr<SEditableTextBox> SourceTextBox;
 
 	/** Current TimecodeSynchronization being used. */
-	UTimecodeSynchronizer* TimecodeSynchronization;
+	TStrongObjectPtr<UTimecodeSynchronizer> TimecodeSynchronization;
 
 	/** Attached Input source index in either TimecodeSynchronization.GetSynchronizedSources() or GetNonSynchronizedSources(). */
 	int32 AttachedSourceIndex;
 
 	/** Whether or not this source is used for synchronization. */
 	bool bIsSynchronizedSource;
-	
-	/** The material that wraps the video texture for display in an SImage. */
-	UMaterial* Material;
-
-	/** The Slate brush that renders the material. */
-	TSharedPtr<FSlateBrush> MaterialBrush;
-
-	/** The video texture sampler in the wrapper material. */
-	UMaterialExpressionTextureSample* TextureSampler;
 };
\ No newline at end of file
diff --git a/Engine/Plugins/Media/TimecodeSynchronizer/Source/TimecodeSynchronizerEditor/TimecodeSynchronizerEditor.build.cs b/Engine/Plugins/Media/TimecodeSynchronizer/Source/TimecodeSynchronizerEditor/TimecodeSynchronizerEditor.build.cs
--- a/Engine/Plugins/Media/TimecodeSynchronizer/Source/TimecodeSynchronizerEditor/TimecodeSynchronizerEditor.build.cs
+++ b/Engine/Plugins/Media/TimecodeSynchronizer/Source/TimecodeSynchronizerEditor/TimecodeSynchronizerEditor.build.cs
@@ -8,12 +8,6 @@ public class TimecodeSynchronizerEditor : ModuleRules
 	{
 		PrivateIncludePaths.Add("TimecodeSynchronizerEditor/Private");
 
-		PublicDependencyModuleNames.AddRange(
-			new string[] {
-				"MediaAssets",
-				"TimecodeSynchronizer",
-			});
-			
 		PrivateDependencyModuleNames.AddRange(
 			new string[] {
 				"AssetTools",
@@ -22,10 +16,13 @@ public class TimecodeSynchronizerEditor : ModuleRules
 				"Engine",
 				"EditorStyle",
 				"InputCore",
-				"PropertyEditor",
-				"SlateCore",
+                "MediaAssets",
+                "MediaPlayerEditor",
+                "PropertyEditor",
+                "SlateCore",
 				"Slate",
-				"UnrealEd",
+                "TimecodeSynchronizer",
+                "UnrealEd",
 			});
 	}
 }
diff --git a/Engine/Plugins/Media/TimecodeSynchronizer/TimecodeSynchronizer.uplugin b/Engine/Plugins/Media/TimecodeSynchronizer/TimecodeSynchronizer.uplugin
--- a/Engine/Plugins/Media/TimecodeSynchronizer/TimecodeSynchronizer.uplugin
+++ b/Engine/Plugins/Media/TimecodeSynchronizer/TimecodeSynchronizer.uplugin
@@ -25,5 +25,11 @@
 			"Type": "Editor",
 			"LoadingPhase": "Default"
 		}
+	],
+	"Plugins": [
+		{
+			"Name": "MediaPlayerEditor",
+			"Enabled": true
+		}
 	]
 }
\ No newline at end of file
diff --git a/Engine/Source/Runtime/TimeManagement/Private/TimeSynchronizationSource.cpp b/Engine/Source/Runtime/TimeManagement/Private/TimeSynchronizationSource.cpp
--- a/Engine/Source/Runtime/TimeManagement/Private/TimeSynchronizationSource.cpp
+++ b/Engine/Source/Runtime/TimeManagement/Private/TimeSynchronizationSource.cpp
@@ -2,9 +2,21 @@
 
 #include "TimeSynchronizationSource.h"
 
+#if WITH_EDITOR
+#include "Widgets/SNullWidget.h"
+#include "Widgets/SWidget.h"
+#endif
+
 UTimeSynchronizationSource::UTimeSynchronizationSource(const FObjectInitializer& ObjectInitializer)
 	: Super(ObjectInitializer)
 	, bUseForSynchronization(true)
 {
 
-}
\ No newline at end of file
+}
+
+#if WITH_EDITOR
+TSharedRef<SWidget> UTimeSynchronizationSource::GetVisualWidget() const 
+{ 
+	return SNullWidget::NullWidget; 
+};
+#endif
\ No newline at end of file
diff --git a/Engine/Source/Runtime/TimeManagement/Public/TimeSynchronizationSource.h b/Engine/Source/Runtime/TimeManagement/Public/TimeSynchronizationSource.h
--- a/Engine/Source/Runtime/TimeManagement/Public/TimeSynchronizationSource.h
+++ b/Engine/Source/Runtime/TimeManagement/Public/TimeSynchronizationSource.h
@@ -10,8 +10,14 @@
 #include "Misc/FrameTime.h"
 #include "Misc/Timecode.h"
 
+#include "Templates/SharedPointer.h"
+
 #include "TimeSynchronizationSource.generated.h"
 
+#if WITH_EDITOR
+class SWidget;
+#endif
+
 struct FTimeSynchronizationOpenData
 {
 	/** Frame rate that will be used as the base for synchronization. */
@@ -64,6 +70,11 @@ public:
 
 public:
 
+#if WITH_EDITOR
+	/** Get Visual Widget of this source to display in UI */
+	virtual TSharedRef<SWidget> GetVisualWidget() const;
+#endif
+
 	/**
 	 * Get the time of the newest available sample (relative to this source's frame rate).
 	 * Note, in cases where Rollover is allowed and has occurred, this may have a lower value than GetOldestSampleTime. 
@@ -186,4 +197,4 @@ public:
 
 		return Offset;
 	}
-};
\ No newline at end of file
+};
diff --git a/Engine/Source/Runtime/TimeManagement/TimeManagement.Build.cs b/Engine/Source/Runtime/TimeManagement/TimeManagement.Build.cs
--- a/Engine/Source/Runtime/TimeManagement/TimeManagement.Build.cs
+++ b/Engine/Source/Runtime/TimeManagement/TimeManagement.Build.cs
@@ -15,5 +15,14 @@ public class TimeManagement : ModuleRules
 				"Engine",
 			}
 		);
-	}
+
+        if (Target.bBuildEditor == true)
+        {
+            PrivateDependencyModuleNames.AddRange(
+                new string[] {
+                "SlateCore",
+                });
+        }
+
+    }
 }
