- Caching NextAvailableSampleTime from MediaQueue in RenderThead to avoid having multiple consumer on this queue

- [Aja/BM] Moving overflow adjustments to gamethread to have single consumer on the media queue

diff --git a/Engine/Plugins/Media/MediaFrameworkUtilities/Source/MediaFrameworkUtilities/Private/MediaBundleTimeSynchronizationSource.cpp b/Engine/Plugins/Media/MediaFrameworkUtilities/Source/MediaFrameworkUtilities/Private/MediaBundleTimeSynchronizationSource.cpp
--- a/Engine/Plugins/Media/MediaFrameworkUtilities/Source/MediaFrameworkUtilities/Private/MediaBundleTimeSynchronizationSource.cpp
+++ b/Engine/Plugins/Media/MediaFrameworkUtilities/Source/MediaFrameworkUtilities/Private/MediaBundleTimeSynchronizationSource.cpp
@@ -73,10 +73,11 @@ FFrameTime UMediaBundleTimeSynchronizationSource::GetNewestSampleTime() const
 		const TSharedPtr<IMediaPlayer, ESPMode::ThreadSafe>& Player = MediaBundle->GetMediaPlayer()->GetPlayerFacade()->GetPlayer();
 		if (Player.IsValid())
 		{
-			//If there is a sample in the Texture, we consider it as the next one to be used/rendered
-			if (MediaBundle->GetMediaTexture()->GetAvailableSampleCount() > 0)
+			//If there is a valid sample timestamp, we consider it as the next one to be used/rendered
+			const FTimespan NextAvailableTime = MediaBundle->GetMediaTexture()->GetNextSampleTime();
+			if (NextAvailableTime != FTimespan::MinValue())
 			{
-				UseTimespan = MediaBundle->GetMediaTexture()->GetNextSampleTime();
+				UseTimespan = NextAvailableTime;
 			}
 
 			if (Player->GetCache().GetSampleCount(EMediaCacheState::Loaded) > 0)
@@ -108,10 +109,11 @@ FFrameTime UMediaBundleTimeSynchronizationSource::GetOldestSampleTime() const
 		const TSharedPtr<IMediaPlayer, ESPMode::ThreadSafe>& Player = MediaBundle->GetMediaPlayer()->GetPlayerFacade()->GetPlayer();
 		if (Player.IsValid())
 		{
-			//If there is a sample in the Texture, we consider it as the next one to be used/rendered
-			if (MediaBundle->GetMediaTexture()->GetAvailableSampleCount() > 0)
+			//If there is a valid sample timestamp, we consider it as the next one to be used/rendered
+			const FTimespan NextAvailableTime = MediaBundle->GetMediaTexture()->GetNextSampleTime();
+			if (NextAvailableTime != FTimespan::MinValue())
 			{
-				UseTimespan = MediaBundle->GetMediaTexture()->GetNextSampleTime();
+				UseTimespan = NextAvailableTime;
 			}
 
 			if (Player->GetCache().GetSampleCount(EMediaCacheState::Loaded) > 0)
diff --git a/Engine/Plugins/Media/TimecodeSynchronizer/Source/TimecodeSynchronizer/Private/MediaPlayerTimeSynchronizationSource.cpp b/Engine/Plugins/Media/TimecodeSynchronizer/Source/TimecodeSynchronizer/Private/MediaPlayerTimeSynchronizationSource.cpp
--- a/Engine/Plugins/Media/TimecodeSynchronizer/Source/TimecodeSynchronizer/Private/MediaPlayerTimeSynchronizationSource.cpp
+++ b/Engine/Plugins/Media/TimecodeSynchronizer/Source/TimecodeSynchronizer/Private/MediaPlayerTimeSynchronizationSource.cpp
@@ -42,14 +42,15 @@ FFrameTime UMediaPlayerTimeSynchronizationSource::GetOldestSampleTime() const
 
 	if (MediaTexture && MediaTexture->GetMediaPlayer())
 	{
-		if (MediaTexture->GetAvailableSampleCount() > 0)
-		{
-			// Ideally, the MediaTexture (or more likely, the TMediaSampleQueue) would be able to track
-			// the current span of samples available. However, that's already prone to some threading issues
-			// and trying to manage more data will only exacerbate that.
+		// Ideally, the MediaTexture (or more likely, the TMediaSampleQueue) would be able to track
+		// the current span of samples available. However, that's already prone to some threading issues
+		// and trying to manage more data will only exacerbate that.
 
-			// Therefore, we can only use the next available sample time.
-			UseTimespan = MediaTexture->GetNextSampleTime();
+		// Therefore, we can only use the next available sample time.
+		const FTimespan NextAvailableTime = MediaTexture->GetNextSampleTime();
+		if (NextAvailableTime != FTimespan::MinValue())
+		{
+			UseTimespan = NextAvailableTime;
 		}
 
 		const TSharedPtr<IMediaPlayer, ESPMode::ThreadSafe>& Player = MediaTexture->GetMediaPlayer()->GetPlayerFacade()->GetPlayer();
@@ -81,14 +82,15 @@ FFrameTime UMediaPlayerTimeSynchronizationSource::GetNewestSampleTime() const
 
 	if (MediaTexture && MediaTexture->GetMediaPlayer())
 	{
-		if (MediaTexture->GetAvailableSampleCount() > 0)
-		{
-			// Ideally, the MediaTexture (or more likely, the TMediaSampleQueue) would be able to track
-			// the current span of samples available. However, that's already prone to some threading issues
-			// and trying to manage more data will only exacerbate that.
+		// Ideally, the MediaTexture (or more likely, the TMediaSampleQueue) would be able to track
+		// the current span of samples available. However, that's already prone to some threading issues
+		// and trying to manage more data will only exacerbate that.
 
-			// Therefore, we can only use the next available sample time.
-			UseTimespan = MediaTexture->GetNextSampleTime();
+		// Therefore, we can only use the next available sample time.
+		const FTimespan NextAvailableTime = MediaTexture->GetNextSampleTime();
+		if (NextAvailableTime != FTimespan::MinValue())
+		{
+			UseTimespan = NextAvailableTime;
 		}
 
 		const TSharedPtr<IMediaPlayer, ESPMode::ThreadSafe>& Player = MediaTexture->GetMediaPlayer()->GetPlayerFacade()->GetPlayer();
diff --git a/Engine/Source/Runtime/MediaAssets/Private/Assets/MediaTexture.cpp b/Engine/Source/Runtime/MediaAssets/Private/Assets/MediaTexture.cpp
--- a/Engine/Source/Runtime/MediaAssets/Private/Assets/MediaTexture.cpp
+++ b/Engine/Source/Runtime/MediaAssets/Private/Assets/MediaTexture.cpp
@@ -63,6 +63,7 @@ UMediaTexture::UMediaTexture(const FObjectInitializer& ObjectInitializer)
 	, DefaultGuid(FGuid::NewGuid())
 	, Dimensions(FIntPoint::ZeroValue)
 	, Size(0)
+	, CachedNextSampleTime(FTimespan::MinValue())
 {
 	NeverStream = true;
 }
@@ -107,6 +108,11 @@ void UMediaTexture::SetMediaPlayer(UMediaPlayer* NewMediaPlayer)
 }
 
 
+void UMediaTexture::CacheNextAvailableSampleTime(FTimespan InNextSampleTime)
+{
+	CachedNextSampleTime = InNextSampleTime;
+}
+
 #if WITH_EDITOR
 
 void UMediaTexture::SetDefaultMediaPlayer(UMediaPlayer* NewMediaPlayer)
@@ -368,16 +374,7 @@ void UMediaTexture::UpdateQueue()
 
 FTimespan UMediaTexture::GetNextSampleTime() const
 {
-	FTimespan SampleTime;
-
-	TSharedPtr<IMediaTextureSample, ESPMode::ThreadSafe> Sample;
-	const bool bHasSucceed = SampleQueue->Peek(Sample);
-	if (bHasSucceed)
-	{
-		SampleTime = Sample->GetTime();
-	}
-
-	return SampleTime;
+	return CachedNextSampleTime;
 }
 
 int32 UMediaTexture::GetAvailableSampleCount() const
diff --git a/Engine/Source/Runtime/MediaAssets/Private/Misc/MediaTextureResource.cpp b/Engine/Source/Runtime/MediaAssets/Private/Misc/MediaTextureResource.cpp
--- a/Engine/Source/Runtime/MediaAssets/Private/Misc/MediaTextureResource.cpp
+++ b/Engine/Source/Runtime/MediaAssets/Private/Misc/MediaTextureResource.cpp
@@ -251,22 +251,8 @@ void FMediaTextureResource::Render(const FRenderParams& Params)
 			}
 			else if (Sample->GetMediaTextureSampleConverter())
 			{
-				
 				CreateOutputRenderTarget(Sample, Params);
-#if MEDIATEXTURERESOURCE_TRACE_RENDER
-				const bool bDecoded = Sample->GetMediaTextureSampleConverter()->Convert(RenderTargetTextureRHI);
-				if (bDecoded == false)
-				{
-					UE_LOG(LogMediaAssets, VeryVerbose, TEXT("TextureResource %p: Unable to decode video sample using hardware acceleration with time %s at time %s"),
-						this,
-						*Sample->GetTime().ToString(TEXT("%h:%m:%s.%t")),
-						*Params.Time.ToString(TEXT("%h:%m:%s.%t"))
-					);
-				}
-#else
 				Sample->GetMediaTextureSampleConverter()->Convert(RenderTargetTextureRHI);
-#endif
-
 			}
 			else if (MediaTextureResource::RequiresConversion(Sample, Params.SrgbOutput))
 			{
@@ -337,6 +323,9 @@ void FMediaTextureResource::Render(const FRenderParams& Params)
 			ClearTexture(Params.ClearColor, Params.SrgbOutput);
 		}
 	}
+	
+	//Cache next available sample time in the MediaTexture owner since we're the only one that can consume from the queue
+	CacheNextAvailableSampleTime(SampleSource);
 
 	// update external texture registration
 	if (!GSupportsImageExternal)
@@ -884,4 +873,18 @@ void FMediaTextureResource::CreateOutputRenderTarget(const TSharedPtr<IMediaText
 	}
 }
 
+void FMediaTextureResource::CacheNextAvailableSampleTime(const TSharedPtr<FMediaTextureSampleSource, ESPMode::ThreadSafe>& InSampleQueue) const
+{
+	FTimespan SampleTime(FTimespan::MinValue());
+
+	if (InSampleQueue.IsValid())
+	{
+		TSharedPtr<IMediaTextureSample, ESPMode::ThreadSafe> Sample;
+		if (InSampleQueue->Peek(Sample))
+		{
+			SampleTime = Sample->GetTime();
+		}
+	}
 
+	Owner.CacheNextAvailableSampleTime(SampleTime);
+}
diff --git a/Engine/Source/Runtime/MediaAssets/Private/Misc/MediaTextureResource.h b/Engine/Source/Runtime/MediaAssets/Private/Misc/MediaTextureResource.h
--- a/Engine/Source/Runtime/MediaAssets/Private/Misc/MediaTextureResource.h
+++ b/Engine/Source/Runtime/MediaAssets/Private/Misc/MediaTextureResource.h
@@ -149,6 +149,13 @@ protected:
 	 */
 	void CreateOutputRenderTarget(const TSharedPtr<IMediaTextureSample, ESPMode::ThreadSafe>& InSample, const FRenderParams& InParams);
 
+	/**
+	 * Caches next available sample from queue in MediaTexture owner to keep single consumer access
+	 *
+	 * @param InSampleQueue SampleQueue to query sample information from
+	 */
+	void CacheNextAvailableSampleTime(const TSharedPtr<FMediaTextureSampleSource, ESPMode::ThreadSafe>& InSampleQueue) const;
+
 private:
 
 	/**
diff --git a/Engine/Source/Runtime/MediaAssets/Public/MediaTexture.h b/Engine/Source/Runtime/MediaAssets/Public/MediaTexture.h
--- a/Engine/Source/Runtime/MediaAssets/Public/MediaTexture.h
+++ b/Engine/Source/Runtime/MediaAssets/Public/MediaTexture.h
@@ -9,6 +9,7 @@
 #include "Math/IntPoint.h"
 #include "MediaSampleQueue.h"
 #include "Misc/Timespan.h"
+#include "Templates/Atomic.h"
 #include "Templates/SharedPointer.h"
 #include "UObject/ObjectMacros.h"
 #include "UObject/ScriptMacros.h"
@@ -93,10 +94,17 @@ public:
 	UFUNCTION(BlueprintCallable, Category="Media|MediaTexture")
 	void SetMediaPlayer(UMediaPlayer* NewMediaPlayer);
 
+	/**
+	 * Caches the next available sample time from the queue when last rendering was made
+	 * @see GetNextSampleTime
+	 */
+	void CacheNextAvailableSampleTime(FTimespan InNextSampleTime);
+
 	/**
 	 * Gets the next sample Time. Only valid if GetAvailableSampleCount is greater than 0
-	 * @return FTimespan of the next sample
-	 * @see GetAvailableSampleCount
+	 * @note This value is cached when last render command was executed to keep single consumer requirements.
+	 * @return FTimespan of the next sample or FTimespan::MinValue if no sample was available in the queue.
+	 * @see GetAvailableSampleCount, CacheNextAvailableSampleTime
 	 */
 	FTimespan GetNextSampleTime() const;
 
@@ -204,4 +212,7 @@ private:
 
 	/** Critical section to protect last rendered guid since it can be read from anywhere. */
 	mutable FCriticalSection CriticalSection;
+
+	/** Next available sample time when last render call was made */
+	TAtomic<FTimespan> CachedNextSampleTime;
 };
