Update TimecodeSynchronizer to use the FApp flags instead of the GEngine flag.

diff --git a/Engine/Plugins/Media/TimecodeSynchronizer/Source/TimecodeSynchronizer/Private/TimecodeSynchronizer.cpp b/Engine/Plugins/Media/TimecodeSynchronizer/Source/TimecodeSynchronizer/Private/TimecodeSynchronizer.cpp
--- a/Engine/Plugins/Media/TimecodeSynchronizer/Source/TimecodeSynchronizer/Private/TimecodeSynchronizer.cpp
+++ b/Engine/Plugins/Media/TimecodeSynchronizer/Source/TimecodeSynchronizer/Private/TimecodeSynchronizer.cpp
@@ -415,10 +415,14 @@ void UTimecodeSynchronizer::Register()
 		}
 		else
 		{
-			PreviousFixedFrameRate = GEngine->FixedFrameRate;
-			bPreviousUseFixedFrameRate = GEngine->bUseFixedFrameRate;
-			GEngine->FixedFrameRate = FixedFrameRate.AsDecimal();
-			GEngine->bUseFixedFrameRate = true;
+			if (FApp::UseFixedTimeStep())
+			{
+				UE_LOG(LogTimecodeSynchronizer, Warning, TEXT("The engine is already in fixed time step."));
+			}
+			PreviousFixedFrameRate = FApp::GetFixedDeltaTime();
+			bPreviousUseFixedFrameRate = FApp::UseFixedTimeStep();
+			FApp::SetFixedDeltaTime(FixedFrameRate.AsDecimal());
+			FApp::SetUseFixedTimeStep(true);
 		}
 
 		if (TimecodeProviderType == ETimecodeSynchronizationTimecodeType::TimecodeProvider)
@@ -468,8 +472,8 @@ void UTimecodeSynchronizer::Unregister()
 
 		if (RegisteredCustomTimeStep == nullptr)
 		{
-			GEngine->FixedFrameRate = PreviousFixedFrameRate;
-			GEngine->bUseFixedFrameRate = bPreviousUseFixedFrameRate;
+			FApp::SetFixedDeltaTime(PreviousFixedFrameRate);
+			FApp::SetUseFixedTimeStep(bPreviousUseFixedFrameRate);
 		}
 		RegisteredCustomTimeStep = nullptr;
 
@@ -1001,9 +1005,9 @@ void UTimecodeSynchronizer::UpdateSourceStates()
 		const FString StateString = SynchronizationStateToString(State);
 		if (InvalidSources.Num() > 0)
 		{
-			for (const FTimecodeSynchronizerActiveTimecodedInputSource* UnreadySource : UnreadySources)
+			for (const FTimecodeSynchronizerActiveTimecodedInputSource* InvalidSource : InvalidSources)
 			{
-				UE_LOG(LogTimecodeSynchronizer, Error, TEXT("Invalid source found unready during State '%s'"), *StateString);
+				UE_LOG(LogTimecodeSynchronizer, Error, TEXT("Source '%s' is not valid during State '%s'"), *(InvalidSource->GetDisplayName()), *StateString);
 			}
 		}
 
