- Fixed bad delta time setup in timecode synchronizer preventing it to work with blackmagic

diff --git a/Engine/Plugins/Media/TimecodeSynchronizer/Source/TimecodeSynchronizer/Private/TimecodeSynchronizer.cpp b/Engine/Plugins/Media/TimecodeSynchronizer/Source/TimecodeSynchronizer/Private/TimecodeSynchronizer.cpp
--- a/Engine/Plugins/Media/TimecodeSynchronizer/Source/TimecodeSynchronizer/Private/TimecodeSynchronizer.cpp
+++ b/Engine/Plugins/Media/TimecodeSynchronizer/Source/TimecodeSynchronizer/Private/TimecodeSynchronizer.cpp
@@ -428,7 +428,7 @@ void UTimecodeSynchronizer::Register()
 			}
 			PreviousFixedFrameRate = FApp::GetFixedDeltaTime();
 			bPreviousUseFixedFrameRate = FApp::UseFixedTimeStep();
-			FApp::SetFixedDeltaTime(FixedFrameRate.AsDecimal());
+			FApp::SetFixedDeltaTime(FixedFrameRate.AsInterval());
 			FApp::SetUseFixedTimeStep(true);
 		}
 
