- Removing UE4 references from some VP modules

diff --git a/Engine/Plugins/Animation/LiveLink/Source/LiveLink/Private/LiveLinkClient.cpp b/Engine/Plugins/Animation/LiveLink/Source/LiveLink/Private/LiveLinkClient.cpp
--- a/Engine/Plugins/Animation/LiveLink/Source/LiveLink/Private/LiveLinkClient.cpp
+++ b/Engine/Plugins/Animation/LiveLink/Source/LiveLink/Private/LiveLinkClient.cpp
@@ -54,13 +54,13 @@ DEFINE_LOG_CATEGORY(LogLiveLink);
 static TAutoConsoleVariable<int32> CVarMaxNewStaticDataPerUpdate(
 	TEXT("LiveLink.Client.MaxNewStaticDataPerUpdate"),
 	64,
-	TEXT("Maximun number of new static data that can be added in a single UE4 frame."),
+	TEXT("Maximun number of new static data that can be added in a single UE frame."),
 	ECVF_Default);
 
 static TAutoConsoleVariable<int32> CVarMaxNewFrameDataPerUpdate(
 	TEXT("LiveLink.Client.MaxNewFrameDataPerUpdate"),
 	64,
-	TEXT("Maximun number of new frame data that can be added in a single UE4 frame."),
+	TEXT("Maximun number of new frame data that can be added in a single UE frame."),
 	ECVF_Default);
 
 
diff --git a/Engine/Plugins/Media/AjaMedia/Source/AjaMedia/Private/Player/AjaMediaPlayer.cpp b/Engine/Plugins/Media/AjaMedia/Source/AjaMedia/Private/Player/AjaMediaPlayer.cpp
--- a/Engine/Plugins/Media/AjaMedia/Source/AjaMedia/Private/Player/AjaMediaPlayer.cpp
+++ b/Engine/Plugins/Media/AjaMedia/Source/AjaMedia/Private/Player/AjaMediaPlayer.cpp
@@ -385,7 +385,7 @@ void FAjaMediaPlayer::VerifyFrameDropCount()
 		}
 		else if (PreviousFrameDropCount > 0)
 		{
-			UE_LOG(LogAjaMedia, Warning, TEXT("Lost %d frames on input %s. UE4 frame rate is too slow and the capture card was not able to send the frame(s) to UE4."), PreviousFrameDropCount, *GetUrl());
+			UE_LOG(LogAjaMedia, Warning, TEXT("Lost %d frames on input %s. Unreal Engine's frame rate is too slow and the capture card was not able to send the frame(s) to Unreal."), PreviousFrameDropCount, *GetUrl());
 			PreviousFrameDropCount = 0;
 		}
 		LastFrameDropCount = FrameDropCount;
diff --git a/Engine/Plugins/Media/AjaMedia/Source/AjaMedia/Public/AjaMediaSource.h b/Engine/Plugins/Media/AjaMedia/Source/AjaMedia/Public/AjaMediaSource.h
--- a/Engine/Plugins/Media/AjaMedia/Source/AjaMedia/Public/AjaMediaSource.h
+++ b/Engine/Plugins/Media/AjaMedia/Source/AjaMedia/Public/AjaMediaSource.h
@@ -19,7 +19,7 @@ enum class EAjaMediaSourceColorFormat : uint8
 };
 
 /**
- * Available number of audio channel supported by UE4 & AJA
+ * Available number of audio channel supported by Unreal Engine & AJA
  */
 UENUM()
 enum class EAjaMediaAudioChannel : uint8
diff --git a/Engine/Plugins/Media/BlackmagicMedia/Source/BlackmagicMedia/Public/BlackmagicMediaSource.h b/Engine/Plugins/Media/BlackmagicMedia/Source/BlackmagicMedia/Public/BlackmagicMediaSource.h
--- a/Engine/Plugins/Media/BlackmagicMedia/Source/BlackmagicMedia/Public/BlackmagicMediaSource.h
+++ b/Engine/Plugins/Media/BlackmagicMedia/Source/BlackmagicMedia/Public/BlackmagicMediaSource.h
@@ -21,7 +21,7 @@ enum class EBlackmagicMediaSourceColorFormat : uint8
 };
 
 /**
- * Available number of audio channel supported by UE4 & Capture card.
+ * Available number of audio channel supported by Unreal Engine & Capture card.
  */
 UENUM(BlueprintType)
 enum class EBlackmagicMediaAudioChannel : uint8
diff --git a/Engine/Plugins/Media/BlackmagicMedia/Source/BlackmagicMediaOutput/Private/BlackmagicMediaCapture.cpp b/Engine/Plugins/Media/BlackmagicMedia/Source/BlackmagicMediaOutput/Private/BlackmagicMediaCapture.cpp
--- a/Engine/Plugins/Media/BlackmagicMedia/Source/BlackmagicMediaOutput/Private/BlackmagicMediaCapture.cpp
+++ b/Engine/Plugins/Media/BlackmagicMedia/Source/BlackmagicMediaOutput/Private/BlackmagicMediaCapture.cpp
@@ -376,7 +376,7 @@ bool UBlackmagicMediaCapture::InitBlackmagic(UBlackmagicMediaOutput* InBlackmagi
 	switch (InBlackmagicMediaOutput->OutputConfiguration.MediaConfiguration.MediaConnection.TransportType)
 	{
 	case EMediaIOTransportType::SingleLink:
-	case EMediaIOTransportType::HDMI: // Blackmagic support HDMI but it is not shown in UE4 UI. It's configured in BMD design tool and it's consider a normal link by UE4.
+	case EMediaIOTransportType::HDMI: // Blackmagic support HDMI but it is not shown in UE's UI. It's configured in BMD design tool and it's considered a normal link by UE.
 		ChannelOptions.LinkConfiguration = BlackmagicDesign::ELinkConfiguration::SingleLink;
 		break;
 	case EMediaIOTransportType::DualLink:
diff --git a/Engine/Plugins/Media/MediaIOFramework/Source/MediaIOCore/Public/MediaOutput.h b/Engine/Plugins/Media/MediaIOFramework/Source/MediaIOCore/Public/MediaOutput.h
--- a/Engine/Plugins/Media/MediaIOFramework/Source/MediaIOCore/Public/MediaOutput.h
+++ b/Engine/Plugins/Media/MediaIOFramework/Source/MediaIOCore/Public/MediaOutput.h
@@ -33,7 +33,7 @@ enum class EMediaCaptureSourceType : int8
  * Abstract base class for media output.
  *
  * Media output describe the location and/or settings of media objects that can
- * be used to output the content of UE4 to a target device via a MediaCapture.
+ * be used to output the content of Unreal Engine to a target device via a MediaCapture.
  */
 UCLASS(Abstract, editinlinenew, BlueprintType, hidecategories = (Object))
 class MEDIAIOCORE_API UMediaOutput : public UObject
diff --git a/Engine/Plugins/Media/TimecodeSynchronizer/Source/TimecodeSynchronizer/Public/TimecodeSynchronizer.h b/Engine/Plugins/Media/TimecodeSynchronizer/Source/TimecodeSynchronizer/Public/TimecodeSynchronizer.h
--- a/Engine/Plugins/Media/TimecodeSynchronizer/Source/TimecodeSynchronizer/Public/TimecodeSynchronizer.h
+++ b/Engine/Plugins/Media/TimecodeSynchronizer/Source/TimecodeSynchronizer/Public/TimecodeSynchronizer.h
@@ -28,7 +28,7 @@ enum class ETimecodeSynchronizationSyncMode
 	 * Synchronization will be achieved once the synchronizer detects all input sources have frames that correspond
 	 * with the offset timecode.
 	 *
-	 * This is suitable for applications trying to keep multiple UE4 instances in sync while using nDisplay / genlock.
+	 * This is suitable for applications trying to keep multiple Unreal Engine instances in sync while using nDisplay / genlock.
 	 */
 	UserDefinedOffset,
 
@@ -36,7 +36,7 @@ enum class ETimecodeSynchronizationSyncMode
 	 * Engine will try and automatically determine an appropriate offset based on what frames are available
 	 * on the given sources.
 	 *
-	 * This is suitable for running a single UE4 instance that just wants to synchronize its inputs.
+	 * This is suitable for running a single Unreal Engine instance that just wants to synchronize its inputs.
 	 */
 	Auto,
 
diff --git a/Engine/Source/Runtime/LiveLinkInterface/Public/ILiveLinkClient.h b/Engine/Source/Runtime/LiveLinkInterface/Public/ILiveLinkClient.h
--- a/Engine/Source/Runtime/LiveLinkInterface/Public/ILiveLinkClient.h
+++ b/Engine/Source/Runtime/LiveLinkInterface/Public/ILiveLinkClient.h
@@ -103,7 +103,7 @@ PRAGMA_ENABLE_DEPRECATION_WARNINGS
 
 
 /**
- * Interface for streaming and consuming data from external sources into UE4.
+ * Interface for streaming and consuming data from external sources into UE.
  * A LiveLinkSource, may stream multiple LiveLinkSubject.
  * Pushing and evaluating data can be executed on any thread. The other functions must be executed on the Game Thread.
  * Subject may shared name between sources, but only 1 of those subjects may be enabled.
diff --git a/Engine/Source/Runtime/LiveLinkMessageBusFramework/Public/LiveLinkProvider.h b/Engine/Source/Runtime/LiveLinkMessageBusFramework/Public/LiveLinkProvider.h
--- a/Engine/Source/Runtime/LiveLinkMessageBusFramework/Public/LiveLinkProvider.h
+++ b/Engine/Source/Runtime/LiveLinkMessageBusFramework/Public/LiveLinkProvider.h
@@ -55,13 +55,13 @@ public:
 	virtual bool UpdateSubjectStaticData(const FName SubjectName, TSubclassOf<ULiveLinkRole> Role, FLiveLinkStaticDataStruct&& StaticData) = 0;
 
 	/**
-	 * Inform UE4 that a subject won't be streamed anymore.
+	 * Inform UE that a subject won't be streamed anymore.
 	 * @param SubjectName	The name of the subject.
 	 */
 	virtual void RemoveSubject(const FName SubjectName) = 0;
 
 	/**
-	 * Send the static data of a subject to UE4.
+	 * Send the static data of a subject to UE.
 	 * @param SubjectName	The name of the subject
 	 * @param StaticData	The frame data of the subject. The type should match the role's data send with UpdateSubjectStaticData.
 							The FLiveLinkFrameDataStruct doesn't have a copy constructor.
