Deprecate TimecodeSynchronizer plugin

diff --git a/Engine/Plugins/Media/TimecodeSynchronizer/Source/TimecodeSynchronizer/Private/TimecodeSynchronizer.cpp b/Engine/Plugins/Media/TimecodeSynchronizer/Source/TimecodeSynchronizer/Private/TimecodeSynchronizer.cpp
--- a/Engine/Plugins/Media/TimecodeSynchronizer/Source/TimecodeSynchronizer/Private/TimecodeSynchronizer.cpp
+++ b/Engine/Plugins/Media/TimecodeSynchronizer/Source/TimecodeSynchronizer/Private/TimecodeSynchronizer.cpp
@@ -16,6 +16,8 @@
 
 #define LOCTEXT_NAMESPACE "TimecodeSynchronizer"
 
+PRAGMA_DISABLE_DEPRECATION_WARNINGS
+
 namespace TimecodeSynchronizerPrivate
 {
 	struct FTimecodeInputSourceValidator
@@ -1052,5 +1054,5 @@ FFrameTime UTimecodeSynchronizer::CalculateSyncTime()
 		return CurrentProviderFrameTime + ActualFrameOffset;
 	}
 }
-
+PRAGMA_ENABLE_DEPRECATION_WARNINGS
 #undef LOCTEXT_NAMESPACE
diff --git a/Engine/Plugins/Media/TimecodeSynchronizer/Source/TimecodeSynchronizer/Private/TimecodeSynchronizerModule.cpp b/Engine/Plugins/Media/TimecodeSynchronizer/Source/TimecodeSynchronizer/Private/TimecodeSynchronizerModule.cpp
--- a/Engine/Plugins/Media/TimecodeSynchronizer/Source/TimecodeSynchronizer/Private/TimecodeSynchronizerModule.cpp
+++ b/Engine/Plugins/Media/TimecodeSynchronizer/Source/TimecodeSynchronizer/Private/TimecodeSynchronizerModule.cpp
@@ -11,6 +11,8 @@
 
 DEFINE_LOG_CATEGORY(LogTimecodeSynchronizer)
 
+PRAGMA_DISABLE_DEPRECATION_WARNINGS
+
 //////////////////////////////////////////////////////////////////////////
 // FTimecodeSynchronizer
 class FTimecodeSynchronizer : public ITimecodeSynchronizerModule, public FSelfRegisteringExec
@@ -68,3 +70,5 @@ public:
 //////////////////////////////////////////////////////////////////////////
 
 IMPLEMENT_MODULE(FTimecodeSynchronizer, TimecodeSynchronizer);
+
+PRAGMA_ENABLE_DEPRECATION_WARNINGS
\ No newline at end of file
diff --git a/Engine/Plugins/Media/TimecodeSynchronizer/Source/TimecodeSynchronizer/Public/TimecodeSynchronizer.h b/Engine/Plugins/Media/TimecodeSynchronizer/Source/TimecodeSynchronizer/Public/TimecodeSynchronizer.h
--- a/Engine/Plugins/Media/TimecodeSynchronizer/Source/TimecodeSynchronizer/Public/TimecodeSynchronizer.h
+++ b/Engine/Plugins/Media/TimecodeSynchronizer/Source/TimecodeSynchronizer/Public/TimecodeSynchronizer.h
@@ -14,13 +14,15 @@
 
 #include "TimecodeSynchronizer.generated.h"
 
+PRAGMA_DISABLE_DEPRECATION_WARNINGS
+
 class UFixedFrameRateCustomTimeStep;
 
 /**
  * Defines the various modes that the synchronizer can use to try and achieve synchronization.
  */
 UENUM()
-enum class ETimecodeSynchronizationSyncMode
+enum class UE_DEPRECATED(5.0, "The TimecodeSynchronizer plugin is deprecated. Please update your project to use the features of the TimedDataMonitor plugin.") ETimecodeSynchronizationSyncMode
 {
 	/**
 	 * User will specify an offset (number of frames) from the Timecode Source (see ETimecodeSycnrhonizationTimecodeType).
@@ -51,7 +53,7 @@ enum class ETimecodeSynchronizationSyncMode
  * Enumerates Timecode source type.
  */
 UENUM()
-enum class ETimecodeSynchronizationTimecodeType
+enum class UE_DEPRECATED(5.0, "The TimecodeSynchronizer plugin is deprecated. Please update your project to use the features of the TimedDataMonitor plugin.") ETimecodeSynchronizationTimecodeType
 {
 	/** Use the configured Engine Default Timecode provider. */
 	DefaultProvider,
@@ -67,7 +69,7 @@ enum class ETimecodeSynchronizationTimecodeType
  * Enumerates possible framerate source
  */
 UENUM()
-enum class ETimecodeSynchronizationFrameRateSources: uint8
+enum class UE_DEPRECATED(5.0, "The TimecodeSynchronizer plugin is deprecated. Please update your project to use the features of the TimedDataMonitor plugin.") ETimecodeSynchronizationFrameRateSources: uint8
 {
 	EngineCustomTimeStepFrameRate UMETA(ToolTip="Frame Rate of engine custom time step if it is of type UFixedFrameRateCustomTimeStep."),
 	CustomFrameRate UMETA(ToolTip = "Custom Frame Rate selected by the user.")
@@ -77,7 +79,7 @@ enum class ETimecodeSynchronizationFrameRateSources: uint8
 /**
  * Enumerates Synchronization related events.
  */
-enum class ETimecodeSynchronizationEvent
+enum class UE_DEPRECATED(5.0, "The TimecodeSynchronizer plugin is deprecated. Please update your project to use the features of the TimedDataMonitor plugin.") ETimecodeSynchronizationEvent
 {
 	/** The synchronization procedure has started. */
 	SynchronizationStarted,
@@ -93,6 +95,7 @@ enum class ETimecodeSynchronizationEvent
 };
 
 /** Cached values to use during synchronization / while synchronized */
+struct UE_DEPRECATED(5.0, "The TimecodeSynchronizer plugin is deprecated. Please update your project to use the features of the TimedDataMonitor plugin.") FTimecodeSynchronizerCachedSyncState;
 struct FTimecodeSynchronizerCachedSyncState
 {
 	/** If we're using rollover, the frame time that represents the rollover point (e.g., the modulus). */
@@ -109,6 +112,7 @@ struct FTimecodeSynchronizerCachedSyncState
 };
 
 /** Cached frame values for a given source. */
+struct UE_DEPRECATED(5.0, "The TimecodeSynchronizer plugin is deprecated. Please update your project to use the features of the TimedDataMonitor plugin.") FTimecodeSourceState;
 struct FTimecodeSourceState
 {
 	/** Frame time of the newest available sample. */
@@ -124,6 +128,7 @@ struct FTimecodeSourceState
  *
  * The values are typically updated once per frame.
  */
+struct UE_DEPRECATED(5.0, "The TimecodeSynchronizer plugin is deprecated. Please update your project to use the features of the TimedDataMonitor plugin.") FTimecodeSynchronizerActiveTimecodedInputSource;
 USTRUCT()
 struct FTimecodeSynchronizerActiveTimecodedInputSource
 {
@@ -225,6 +230,8 @@ private:
 	FTimecodeSynchronizerActiveTimecodedInputSource& operator=(FTimecodeSynchronizerActiveTimecodedInputSource&&) = delete;
 };
 
+template<> struct UE_DEPRECATED(5.0, "The TimecodeSynchronizer plugin is deprecated. Please update your project to use the features of the TimedDataMonitor plugin.") TStructOpsTypeTraits<FTimecodeSynchronizerActiveTimecodedInputSource>;
+
 template<>
 struct TStructOpsTypeTraits<FTimecodeSynchronizerActiveTimecodedInputSource> : public TStructOpsTypeTraitsBase2<FTimecodeSynchronizerActiveTimecodedInputSource>
 {
@@ -247,6 +254,7 @@ struct TStructOpsTypeTraits<FTimecodeSynchronizerActiveTimecodedInputSource> : p
  * TimeSynchronizationSources). Instead, the synchronizer simply acts as a coordinator
  * making sure all sources are ready, determining if sync is possible, etc.
  */
+class UE_DEPRECATED(5.0, "The TimecodeSynchronizer plugin is deprecated. Please update your project to use the features of the TimedDataMonitor plugin.") UTimecodeSynchronizer;
 UCLASS()
 class TIMECODESYNCHRONIZER_API UTimecodeSynchronizer : public UTimecodeProvider, public FTickableGameObject
 {
@@ -578,3 +586,4 @@ private:
 	bool bAddSourcesGuard;
 	bool bShouldResetTimecodeProvider;
 };
+PRAGMA_ENABLE_DEPRECATION_WARNINGS
diff --git a/Engine/Plugins/Media/TimecodeSynchronizer/Source/TimecodeSynchronizer/Public/TimecodeSynchronizerModule.h b/Engine/Plugins/Media/TimecodeSynchronizer/Source/TimecodeSynchronizer/Public/TimecodeSynchronizerModule.h
--- a/Engine/Plugins/Media/TimecodeSynchronizer/Source/TimecodeSynchronizer/Public/TimecodeSynchronizerModule.h
+++ b/Engine/Plugins/Media/TimecodeSynchronizer/Source/TimecodeSynchronizer/Public/TimecodeSynchronizerModule.h
@@ -11,6 +11,7 @@
 
 DECLARE_LOG_CATEGORY_EXTERN(LogTimecodeSynchronizer, Verbose, All);
 
+class UE_DEPRECATED(5.0, "The TimecodeSynchronizer plugin is deprecated. Please update your project to use the features of the TimedDataMonitor plugin.") ITimecodeSynchronizerModule;
 class TIMECODESYNCHRONIZER_API ITimecodeSynchronizerModule : public IModuleInterface
 {
 public:
diff --git a/Engine/Plugins/Media/TimecodeSynchronizer/Source/TimecodeSynchronizer/Public/TimecodeSynchronizerProjectSettings.h b/Engine/Plugins/Media/TimecodeSynchronizer/Source/TimecodeSynchronizer/Public/TimecodeSynchronizerProjectSettings.h
--- a/Engine/Plugins/Media/TimecodeSynchronizer/Source/TimecodeSynchronizer/Public/TimecodeSynchronizerProjectSettings.h
+++ b/Engine/Plugins/Media/TimecodeSynchronizer/Source/TimecodeSynchronizer/Public/TimecodeSynchronizerProjectSettings.h
@@ -9,10 +9,11 @@
 
 #include "TimecodeSynchronizerProjectSettings.generated.h"
 
-
+PRAGMA_DISABLE_DEPRECATION_WARNINGS
 /**
  * Global settings for TimecodeSynchronizer
  */
+class UE_DEPRECATED(5.0, "The TimecodeSynchronizer plugin is deprecated. Please update your project to use the features of the TimedDataMonitor plugin.") UTimecodeSynchronizerProjectSettings;
 UCLASS(config=Engine)
 class TIMECODESYNCHRONIZER_API UTimecodeSynchronizerProjectSettings : public UObject
 {
@@ -49,6 +50,7 @@ public:
 /**
  * Editor settings for TimecodeSynchronizer
  */
+class UE_DEPRECATED(5.0, "The TimecodeSynchronizer plugin is deprecated. Please update your project to use the features of the TimedDataMonitor plugin.") UTimecodeSynchronizerEditorSettings;
 UCLASS(config= EditorPerProjectUserSettings)
 class TIMECODESYNCHRONIZER_API UTimecodeSynchronizerEditorSettings : public UObject
 {
@@ -60,3 +62,4 @@ public:
 	UPROPERTY(config, EditAnywhere, Category="TimecodeSynchronizer")
 	TSoftObjectPtr<UTimecodeSynchronizer> UserTimecodeSynchronizer;
 };
+PRAGMA_ENABLE_DEPRECATION_WARNINGS
\ No newline at end of file
diff --git a/Engine/Plugins/Media/TimecodeSynchronizer/Source/TimecodeSynchronizerEditor/Private/AssetEditor/TimecodeSynchronizerEditorToolkit.cpp b/Engine/Plugins/Media/TimecodeSynchronizer/Source/TimecodeSynchronizerEditor/Private/AssetEditor/TimecodeSynchronizerEditorToolkit.cpp
--- a/Engine/Plugins/Media/TimecodeSynchronizer/Source/TimecodeSynchronizerEditor/Private/AssetEditor/TimecodeSynchronizerEditorToolkit.cpp
+++ b/Engine/Plugins/Media/TimecodeSynchronizer/Source/TimecodeSynchronizerEditor/Private/AssetEditor/TimecodeSynchronizerEditorToolkit.cpp
@@ -31,6 +31,8 @@ namespace TimecodeSynchronizerEditorToolkit
 	const FName Layout(TEXT("Standalone_TimecodeSynchronizerEditor_Layout_v1"));
 }
 
+PRAGMA_DISABLE_DEPRECATION_WARNINGS
+
 TSharedRef<FTimecodeSynchronizerEditorToolkit> FTimecodeSynchronizerEditorToolkit::CreateEditor(const EToolkitMode::Type Mode, const TSharedPtr< class IToolkitHost >& InitToolkitHost, class UTimecodeSynchronizer* InTimecodeSynchronizer)
 {
 	TSharedRef<FTimecodeSynchronizerEditorToolkit> NewEditor(new FTimecodeSynchronizerEditorToolkit());
@@ -410,4 +412,6 @@ void FTimecodeSynchronizerEditorToolkit::RemoveEditingObject(UObject* Object)
 	Super::RemoveEditingObject(Object);
 }
 
+PRAGMA_ENABLE_DEPRECATION_WARNINGS
+
 #undef LOCTEXT_NAMESPACE
diff --git a/Engine/Plugins/Media/TimecodeSynchronizer/Source/TimecodeSynchronizerEditor/Private/AssetEditor/TimecodeSynchronizerEditorToolkit.h b/Engine/Plugins/Media/TimecodeSynchronizer/Source/TimecodeSynchronizerEditor/Private/AssetEditor/TimecodeSynchronizerEditorToolkit.h
--- a/Engine/Plugins/Media/TimecodeSynchronizer/Source/TimecodeSynchronizerEditor/Private/AssetEditor/TimecodeSynchronizerEditorToolkit.h
+++ b/Engine/Plugins/Media/TimecodeSynchronizer/Source/TimecodeSynchronizerEditor/Private/AssetEditor/TimecodeSynchronizerEditorToolkit.h
@@ -7,6 +7,8 @@
 
 #include "TimecodeSynchronizer.h"
 
+PRAGMA_DISABLE_DEPRECATION_WARNINGS
+
 /** Viewer/editor for a TimecodeSynchronizer */
 class FTimecodeSynchronizerEditorToolkit : public FAssetEditorToolkit
 {
@@ -55,3 +57,4 @@ private:
 	/** Details view */
 	TSharedPtr<class IDetailsView> DetailsView;
 };
+PRAGMA_ENABLE_DEPRECATION_WARNINGS
\ No newline at end of file
diff --git a/Engine/Plugins/Media/TimecodeSynchronizer/Source/TimecodeSynchronizerEditor/Private/AssetTypeActions/AssetTypeActions_TimecodeSynchronizer.cpp b/Engine/Plugins/Media/TimecodeSynchronizer/Source/TimecodeSynchronizerEditor/Private/AssetTypeActions/AssetTypeActions_TimecodeSynchronizer.cpp
--- a/Engine/Plugins/Media/TimecodeSynchronizer/Source/TimecodeSynchronizerEditor/Private/AssetTypeActions/AssetTypeActions_TimecodeSynchronizer.cpp
+++ b/Engine/Plugins/Media/TimecodeSynchronizer/Source/TimecodeSynchronizerEditor/Private/AssetTypeActions/AssetTypeActions_TimecodeSynchronizer.cpp
@@ -6,7 +6,7 @@
 #include "TimecodeSynchronizer.h"
 
 #define LOCTEXT_NAMESPACE "TimecodeSynchronizerEditor"
-
+PRAGMA_DISABLE_DEPRECATION_WARNINGS
 FText FAssetTypeActions_TimecodeSynchronizer::GetName() const
 {
 	return LOCTEXT("AssetTypeActions_TimecodeSynchronizer", "Timecode Synchronizer");
@@ -27,5 +27,5 @@ void FAssetTypeActions_TimecodeSynchronizer::OpenAssetEditor(const TArray<UObjec
 		}
 	}
 }
-
+PRAGMA_ENABLE_DEPRECATION_WARNINGS
 #undef LOCTEXT_NAMESPACE
diff --git a/Engine/Plugins/Media/TimecodeSynchronizer/Source/TimecodeSynchronizerEditor/Private/Factory/TimecodeSynchronizerFactory.cpp b/Engine/Plugins/Media/TimecodeSynchronizer/Source/TimecodeSynchronizerEditor/Private/Factory/TimecodeSynchronizerFactory.cpp
--- a/Engine/Plugins/Media/TimecodeSynchronizer/Source/TimecodeSynchronizerEditor/Private/Factory/TimecodeSynchronizerFactory.cpp
+++ b/Engine/Plugins/Media/TimecodeSynchronizer/Source/TimecodeSynchronizerEditor/Private/Factory/TimecodeSynchronizerFactory.cpp
@@ -4,6 +4,7 @@
 
 #include "TimecodeSynchronizer.h"
 
+PRAGMA_DISABLE_DEPRECATION_WARNINGS
 UTimecodeSynchronizerFactory::UTimecodeSynchronizerFactory(const FObjectInitializer& ObjectInitializer)
 	: Super(ObjectInitializer)
 {
@@ -19,3 +20,4 @@ UObject* UTimecodeSynchronizerFactory::FactoryCreateNew(UClass* Class, UObject*
 	UTimecodeSynchronizer* TimecodeSynchronizer = NewObject<UTimecodeSynchronizer>(InParent, Name, Flags);
 	return TimecodeSynchronizer;
 }
+PRAGMA_ENABLE_DEPRECATION_WARNINGS
\ No newline at end of file
diff --git a/Engine/Plugins/Media/TimecodeSynchronizer/Source/TimecodeSynchronizerEditor/Private/TimecodeSynchronizerEditorModule.cpp b/Engine/Plugins/Media/TimecodeSynchronizer/Source/TimecodeSynchronizerEditor/Private/TimecodeSynchronizerEditorModule.cpp
--- a/Engine/Plugins/Media/TimecodeSynchronizer/Source/TimecodeSynchronizerEditor/Private/TimecodeSynchronizerEditorModule.cpp
+++ b/Engine/Plugins/Media/TimecodeSynchronizer/Source/TimecodeSynchronizerEditor/Private/TimecodeSynchronizerEditorModule.cpp
@@ -13,6 +13,8 @@
 
 #define LOCTEXT_NAMESPACE "TimecodeSynchronizerEditor"
 
+PRAGMA_DISABLE_DEPRECATION_WARNINGS
+
 //////////////////////////////////////////////////////////////////////////
 // FTimecodeSynchronizerEditorModule
 class FTimecodeSynchronizerEditorModule : public ITimecodeSynchronizerEditorModule
@@ -68,4 +70,6 @@ public:
 
 IMPLEMENT_MODULE(FTimecodeSynchronizerEditorModule, TimecodeSynchronizerEditor);
 
+PRAGMA_ENABLE_DEPRECATION_WARNINGS
+
 #undef LOCTEXT_NAMESPACE
diff --git a/Engine/Plugins/Media/TimecodeSynchronizer/Source/TimecodeSynchronizerEditor/Private/UI/TimecodeSynchronizerEditorCommand.cpp b/Engine/Plugins/Media/TimecodeSynchronizer/Source/TimecodeSynchronizerEditor/Private/UI/TimecodeSynchronizerEditorCommand.cpp
--- a/Engine/Plugins/Media/TimecodeSynchronizer/Source/TimecodeSynchronizerEditor/Private/UI/TimecodeSynchronizerEditorCommand.cpp
+++ b/Engine/Plugins/Media/TimecodeSynchronizer/Source/TimecodeSynchronizerEditor/Private/UI/TimecodeSynchronizerEditorCommand.cpp
@@ -10,6 +10,8 @@
 
 #define LOCTEXT_NAMESPACE "TimecodeSynchronizerEditor"
 
+PRAGMA_DISABLE_DEPRECATION_WARNINGS
+
 //////////////////////////////////////////////////////////////////////////
 // FTimecodeSynchronizerEditorCommand
 
@@ -57,4 +59,6 @@ bool FTimecodeSynchronizerEditorCommand::CanOpenTimecodeSynchronizerEditor()
 	return !UTimecodeSynchronizerProjectSettings::StaticClass()->GetDefaultObject<UTimecodeSynchronizerProjectSettings>()->DefaultTimecodeSynchronizer.IsNull();
 }
 
+PRAGMA_ENABLE_DEPRECATION_WARNINGS
+
 #undef LOCTEXT_NAMESPACE
diff --git a/Engine/Plugins/Media/TimecodeSynchronizer/Source/TimecodeSynchronizerEditor/Private/UI/TimecodeSynchronizerEditorLevelToolbar.cpp b/Engine/Plugins/Media/TimecodeSynchronizer/Source/TimecodeSynchronizerEditor/Private/UI/TimecodeSynchronizerEditorLevelToolbar.cpp
--- a/Engine/Plugins/Media/TimecodeSynchronizer/Source/TimecodeSynchronizerEditor/Private/UI/TimecodeSynchronizerEditorLevelToolbar.cpp
+++ b/Engine/Plugins/Media/TimecodeSynchronizer/Source/TimecodeSynchronizerEditor/Private/UI/TimecodeSynchronizerEditorLevelToolbar.cpp
@@ -22,6 +22,8 @@
 
 #define LOCTEXT_NAMESPACE "TimecodeSynchronizerEditor"
 
+PRAGMA_DISABLE_DEPRECATION_WARNINGS
+
 //////////////////////////////////////////////////////////////////////////
 // FTimecodeSynchronizerEditorLevelToolbar
 
@@ -198,5 +200,5 @@ void FTimecodeSynchronizerEditorLevelToolbar::NewTimecodeSynchronizerSelected(co
 	GetMutableDefault<UTimecodeSynchronizerEditorSettings>()->SaveConfig();
 	CurrentTimecodeSynchronizer = Asset;
 }
-
+PRAGMA_ENABLE_DEPRECATION_WARNINGS
 #undef LOCTEXT_NAMESPACE
diff --git a/Engine/Plugins/Media/TimecodeSynchronizer/Source/TimecodeSynchronizerEditor/Private/UI/TimecodeSynchronizerEditorLevelToolbar.h b/Engine/Plugins/Media/TimecodeSynchronizer/Source/TimecodeSynchronizerEditor/Private/UI/TimecodeSynchronizerEditorLevelToolbar.h
--- a/Engine/Plugins/Media/TimecodeSynchronizer/Source/TimecodeSynchronizerEditor/Private/UI/TimecodeSynchronizerEditorLevelToolbar.h
+++ b/Engine/Plugins/Media/TimecodeSynchronizer/Source/TimecodeSynchronizerEditor/Private/UI/TimecodeSynchronizerEditorLevelToolbar.h
@@ -15,7 +15,7 @@ class SWidget;
 
 //////////////////////////////////////////////////////////////////////////
 // FTimecodeSynchronizerEditorLevelToolbar
-
+PRAGMA_DISABLE_DEPRECATION_WARNINGS
 class FTimecodeSynchronizerEditorLevelToolbar
 {
 public:
@@ -37,4 +37,4 @@ private:
 	TSharedPtr<FExtender> LevelToolbarExtender;
 	TWeakObjectPtr<UTimecodeSynchronizer> CurrentTimecodeSynchronizer;
 };
-
+PRAGMA_ENABLE_DEPRECATION_WARNINGS
diff --git a/Engine/Plugins/Media/TimecodeSynchronizer/Source/TimecodeSynchronizerEditor/Private/Widgets/STimecodeSynchronizerSourceViewer.cpp b/Engine/Plugins/Media/TimecodeSynchronizer/Source/TimecodeSynchronizerEditor/Private/Widgets/STimecodeSynchronizerSourceViewer.cpp
--- a/Engine/Plugins/Media/TimecodeSynchronizer/Source/TimecodeSynchronizerEditor/Private/Widgets/STimecodeSynchronizerSourceViewer.cpp
+++ b/Engine/Plugins/Media/TimecodeSynchronizer/Source/TimecodeSynchronizerEditor/Private/Widgets/STimecodeSynchronizerSourceViewer.cpp
@@ -11,7 +11,7 @@
 #include "Widgets/SOverlay.h"
 #include "Widgets/STimecodeSynchronizerSourceViewport.h"
 
-
+PRAGMA_DISABLE_DEPRECATION_WARNINGS
 /* STimecodeSynchronizerSourceViewer structors
  *****************************************************************************/
 
@@ -111,4 +111,4 @@ void STimecodeSynchronizerSourceViewer::HandleSynchronizationEvent(ETimecodeSync
 	}
 }
 
-
+PRAGMA_ENABLE_DEPRECATION_WARNINGS
diff --git a/Engine/Plugins/Media/TimecodeSynchronizer/Source/TimecodeSynchronizerEditor/Private/Widgets/STimecodeSynchronizerSourceViewer.h b/Engine/Plugins/Media/TimecodeSynchronizer/Source/TimecodeSynchronizerEditor/Private/Widgets/STimecodeSynchronizerSourceViewer.h
--- a/Engine/Plugins/Media/TimecodeSynchronizer/Source/TimecodeSynchronizerEditor/Private/Widgets/STimecodeSynchronizerSourceViewer.h
+++ b/Engine/Plugins/Media/TimecodeSynchronizer/Source/TimecodeSynchronizerEditor/Private/Widgets/STimecodeSynchronizerSourceViewer.h
@@ -17,7 +17,7 @@ class SEditableTextBox;
 class STimecodeSynchronizerSourceViewport;
 class SVerticalBox;
 
-
+PRAGMA_DISABLE_DEPRECATION_WARNINGS
 /**
  * Implements the contents of the viewer tab in the TimecodeSynchronizer editor.
  */
@@ -61,4 +61,4 @@ private:
 	/** VerticalBox holding active sources viewport */
 	TSharedPtr<SVerticalBox> ViewportVerticalBox;
 };
-	
\ No newline at end of file
+PRAGMA_ENABLE_DEPRECATION_WARNINGS
\ No newline at end of file
diff --git a/Engine/Plugins/Media/TimecodeSynchronizer/Source/TimecodeSynchronizerEditor/Private/Widgets/STimecodeSynchronizerSourceViewport.cpp b/Engine/Plugins/Media/TimecodeSynchronizer/Source/TimecodeSynchronizerEditor/Private/Widgets/STimecodeSynchronizerSourceViewport.cpp
--- a/Engine/Plugins/Media/TimecodeSynchronizer/Source/TimecodeSynchronizerEditor/Private/Widgets/STimecodeSynchronizerSourceViewport.cpp
+++ b/Engine/Plugins/Media/TimecodeSynchronizer/Source/TimecodeSynchronizerEditor/Private/Widgets/STimecodeSynchronizerSourceViewport.cpp
@@ -17,6 +17,7 @@
 
 #define LOCTEXT_NAMESPACE "STimecodeSynchronizercSourceViewport"
 
+PRAGMA_DISABLE_DEPRECATION_WARNINGS
 
 /* STimecodeSynchronizerSourceViewport structors
 *****************************************************************************/
@@ -220,4 +221,6 @@ const FTimecodeSynchronizerActiveTimecodedInputSource* STimecodeSynchronizerSour
 	return Result;
 }
 
+PRAGMA_ENABLE_DEPRECATION_WARNINGS
+
 #undef LOCTEXT_NAMESPACE
diff --git a/Engine/Plugins/Media/TimecodeSynchronizer/Source/TimecodeSynchronizerEditor/Private/Widgets/STimecodeSynchronizerSourceViewport.h b/Engine/Plugins/Media/TimecodeSynchronizer/Source/TimecodeSynchronizerEditor/Private/Widgets/STimecodeSynchronizerSourceViewport.h
--- a/Engine/Plugins/Media/TimecodeSynchronizer/Source/TimecodeSynchronizerEditor/Private/Widgets/STimecodeSynchronizerSourceViewport.h
+++ b/Engine/Plugins/Media/TimecodeSynchronizer/Source/TimecodeSynchronizerEditor/Private/Widgets/STimecodeSynchronizerSourceViewport.h
@@ -12,6 +12,8 @@ struct FSlateBrush;
 struct FTimecodeSynchronizerActiveTimecodedInputSource;
 class UTimecodeSynchronizer;
 
+PRAGMA_DISABLE_DEPRECATION_WARNINGS
+
 class STimecodeSynchronizerSourceViewport
 	: public SCompoundWidget
 {
@@ -67,4 +69,5 @@ private:
 
 	/** Whether or not this source is used for synchronization. */
 	bool bIsSynchronizedSource;
-};
\ No newline at end of file
+};
+PRAGMA_ENABLE_DEPRECATION_WARNINGS
\ No newline at end of file
diff --git a/Engine/Plugins/Media/TimecodeSynchronizer/Source/TimecodeSynchronizerEditor/Private/Widgets/STimecodeSynchronizerWidget.cpp b/Engine/Plugins/Media/TimecodeSynchronizer/Source/TimecodeSynchronizerEditor/Private/Widgets/STimecodeSynchronizerWidget.cpp
--- a/Engine/Plugins/Media/TimecodeSynchronizer/Source/TimecodeSynchronizerEditor/Private/Widgets/STimecodeSynchronizerWidget.cpp
+++ b/Engine/Plugins/Media/TimecodeSynchronizer/Source/TimecodeSynchronizerEditor/Private/Widgets/STimecodeSynchronizerWidget.cpp
@@ -22,6 +22,8 @@ namespace TimecodeSynchronizerBarWidget
 	const int32 TextSize = 10;
 }
 
+PRAGMA_DISABLE_DEPRECATION_WARNINGS
+
 STimecodeSynchronizerBarWidget::STimecodeSynchronizerBarWidget()
 	: SWidget()
 	, CurrentFrameWidth(10)
@@ -401,3 +403,4 @@ void STimecodeSynchronizerWidget::SetFrameWidth(int32 InFrameWidth)
 {
 	BarWidget->SetFrameWidth(InFrameWidth);
 }
+PRAGMA_ENABLE_DEPRECATION_WARNINGS
\ No newline at end of file
diff --git a/Engine/Plugins/Media/TimecodeSynchronizer/Source/TimecodeSynchronizerEditor/Private/Widgets/STimecodeSynchronizerWidget.h b/Engine/Plugins/Media/TimecodeSynchronizer/Source/TimecodeSynchronizerEditor/Private/Widgets/STimecodeSynchronizerWidget.h
--- a/Engine/Plugins/Media/TimecodeSynchronizer/Source/TimecodeSynchronizerEditor/Private/Widgets/STimecodeSynchronizerWidget.h
+++ b/Engine/Plugins/Media/TimecodeSynchronizer/Source/TimecodeSynchronizerEditor/Private/Widgets/STimecodeSynchronizerWidget.h
@@ -15,6 +15,7 @@
 class UTimecodeSynchronizer;
 class FSlateFontMeasure;
 
+PRAGMA_DISABLE_DEPRECATION_WARNINGS
 class STimecodeSynchronizerBarWidget : public SWidget
 {
 public:
@@ -104,4 +105,4 @@ public:
 private:
 	TSharedPtr<STimecodeSynchronizerBarWidget> BarWidget;
 };
-
+PRAGMA_ENABLE_DEPRECATION_WARNINGS
diff --git a/Engine/Plugins/Media/TimecodeSynchronizer/Source/TimecodeSynchronizerEditor/Public/TimecodeSynchronizerEditorModule.h b/Engine/Plugins/Media/TimecodeSynchronizer/Source/TimecodeSynchronizerEditor/Public/TimecodeSynchronizerEditorModule.h
--- a/Engine/Plugins/Media/TimecodeSynchronizer/Source/TimecodeSynchronizerEditor/Public/TimecodeSynchronizerEditorModule.h
+++ b/Engine/Plugins/Media/TimecodeSynchronizer/Source/TimecodeSynchronizerEditor/Public/TimecodeSynchronizerEditorModule.h
@@ -8,6 +8,7 @@
 
 #define TIMECODESYNCHRONIZEREDITOR_MODULE_NAME TEXT("TimecodeSynchronizerEditor")
 
+class UE_DEPRECATED(5.0, "The TimecodeSynchronizer plugin is deprecated. Please update your project to use the features of the TimedDataMonitor plugin.") ITimecodeSynchronizerEditorModule;
 class TIMECODESYNCHRONIZEREDITOR_API ITimecodeSynchronizerEditorModule : public IModuleInterface
 {
 public:
diff --git a/Engine/Plugins/Media/TimecodeSynchronizer/Source/TimecodeSynchronizerEditor/Public/TimecodeSynchronizerFactory.h b/Engine/Plugins/Media/TimecodeSynchronizer/Source/TimecodeSynchronizerEditor/Public/TimecodeSynchronizerFactory.h
--- a/Engine/Plugins/Media/TimecodeSynchronizer/Source/TimecodeSynchronizerEditor/Public/TimecodeSynchronizerFactory.h
+++ b/Engine/Plugins/Media/TimecodeSynchronizer/Source/TimecodeSynchronizerEditor/Public/TimecodeSynchronizerFactory.h
@@ -7,6 +7,7 @@
 #include "Factories/Factory.h"
 #include "TimecodeSynchronizerFactory.generated.h"
 
+class UE_DEPRECATED(5.0, "The TimecodeSynchronizer plugin is deprecated. Please update your project to use the features of the TimedDataMonitor plugin.") UTimecodeSynchronizerFactory;
 UCLASS(hidecategories = Object)
 class TIMECODESYNCHRONIZEREDITOR_API UTimecodeSynchronizerFactory : public UFactory
 {
diff --git a/Engine/Plugins/Media/TimecodeSynchronizer/TimecodeSynchronizer.uplugin b/Engine/Plugins/Media/TimecodeSynchronizer/TimecodeSynchronizer.uplugin
--- a/Engine/Plugins/Media/TimecodeSynchronizer/TimecodeSynchronizer.uplugin
+++ b/Engine/Plugins/Media/TimecodeSynchronizer/TimecodeSynchronizer.uplugin
@@ -2,8 +2,8 @@
 	"FileVersion": 3,
 	"Version": 1,
 	"VersionName": "0.1",
-	"FriendlyName": "Timecode Synchronizer",
-	"Description": "An asset that will become the TimecodeProvider once all the inputs get synchronized to a timecode.",
+	"FriendlyName": "Timecode Synchronizer (Deprecated)",
+	"Description": "This plugin has been deprecated and will be removed in a future engine version. Please update your project to use the features of the TimedDataMonitor plugin instead.\nAn asset that will become the TimecodeProvider once all the inputs get synchronized to a timecode.",
 	"Category": "Virtual Production",
 	"CreatedBy": "Epic Games, Inc.",
 	"CreatedByURL": "http://epicgames.com",
diff --git a/Templates/TP_ME_BlankBP/TP_ME_BlankBP.uproject b/Templates/TP_ME_BlankBP/TP_ME_BlankBP.uproject
--- a/Templates/TP_ME_BlankBP/TP_ME_BlankBP.uproject
+++ b/Templates/TP_ME_BlankBP/TP_ME_BlankBP.uproject
@@ -122,10 +122,6 @@
 			"Name": "PixelStreaming",
 			"Enabled": true
 		},
-		{
-			"Name": "TimecodeSynchronizer",
-			"Enabled": true
-		},
 		{
 			"Name": "VariantManager",
 			"Enabled": true
diff --git a/Templates/TP_ME_VProdBP/TP_ME_VprodBP.uproject b/Templates/TP_ME_VProdBP/TP_ME_VprodBP.uproject
--- a/Templates/TP_ME_VProdBP/TP_ME_VprodBP.uproject
+++ b/Templates/TP_ME_VProdBP/TP_ME_VprodBP.uproject
@@ -122,10 +122,6 @@
 			"Name": "PixelStreaming",
 			"Enabled": true
 		},
-		{
-			"Name": "TimecodeSynchronizer",
-			"Enabled": true
-		},
 		{
 			"Name": "VariantManager",
 			"Enabled": true
