Removed banned word from Media IO

diff --git a/Engine/Plugins/Media/TimecodeSynchronizer/Config/BaseTimecodeSynchronizer.ini b/Engine/Plugins/Media/TimecodeSynchronizer/Config/BaseTimecodeSynchronizer.ini
new file mode 100644
--- /dev/null
+++ b/Engine/Plugins/Media/TimecodeSynchronizer/Config/BaseTimecodeSynchronizer.ini
@@ -0,0 +1,2 @@
+[CoreRedirects]
++PropertyRedirects=(OldName="TimecodeSynchronizer.MasterSynchronizationSourceIndex",NewName="TimecodeSynchronizer.MainSynchronizationSourceIndex")
\ No newline at end of file
diff --git a/Engine/Plugins/Media/TimecodeSynchronizer/Source/TimecodeSynchronizer/Private/TimecodeSynchronizer.cpp b/Engine/Plugins/Media/TimecodeSynchronizer/Source/TimecodeSynchronizer/Private/TimecodeSynchronizer.cpp
--- a/Engine/Plugins/Media/TimecodeSynchronizer/Source/TimecodeSynchronizer/Private/TimecodeSynchronizer.cpp
+++ b/Engine/Plugins/Media/TimecodeSynchronizer/Source/TimecodeSynchronizer/Private/TimecodeSynchronizer.cpp
@@ -195,7 +195,7 @@ UTimecodeSynchronizer::UTimecodeSynchronizer()
 	: FixedFrameRate(30, 1)
 	, TimecodeProviderType(ETimecodeSynchronizationTimecodeType::TimecodeProvider)
 	, TimecodeProvider(nullptr)
-	, MasterSynchronizationSourceIndex(INDEX_NONE)
+	, MainSynchronizationSourceIndex(INDEX_NONE)
 	, PreRollingTimecodeMarginOfErrors(4)
 	, PreRollingTimeout(30.f)
 	, bIsTickEnabled(false)
@@ -204,7 +204,7 @@ UTimecodeSynchronizer::UTimecodeSynchronizer()
 	, bRegistered(false)
 	, PreviousFixedFrameRate(0.f)
 	, bPreviousUseFixedFrameRate(false)
-	, ActiveMasterSynchronizationTimecodedSourceIndex(INDEX_NONE)
+	, ActiveMainSynchronizationTimecodedSourceIndex(INDEX_NONE)
 	, bFailGuard(false)
 	, bAddSourcesGuard(false)
 	, bShouldResetTimecodeProvider(false)
@@ -238,7 +238,7 @@ bool UTimecodeSynchronizer::CanEditChange(const FProperty* InProperty) const
 	{
 		return FrameRateSource == ETimecodeSynchronizationFrameRateSources::CustomFrameRate;
 	}
-	else if (PropertyName == GET_MEMBER_NAME_CHECKED(UTimecodeSynchronizer, MasterSynchronizationSourceIndex))
+	else if (PropertyName == GET_MEMBER_NAME_CHECKED(UTimecodeSynchronizer, MainSynchronizationSourceIndex))
 	{
 		return TimecodeProviderType == ETimecodeSynchronizationTimecodeType::InputSource;
 	}
@@ -257,15 +257,15 @@ bool UTimecodeSynchronizer::CanEditChange(const FProperty* InProperty) const
 
 void UTimecodeSynchronizer::PostEditChangeChainProperty(FPropertyChangedChainEvent& PropertyChangedEvent)
 {
-	// Make sure the master source index is valid
+	// Make sure the main source index is valid
 	if (TimecodeProviderType == ETimecodeSynchronizationTimecodeType::InputSource)
 	{
-		if (!TimeSynchronizationInputSources.IsValidIndex(MasterSynchronizationSourceIndex)
-			|| TimeSynchronizationInputSources[MasterSynchronizationSourceIndex] == nullptr
-			|| !TimeSynchronizationInputSources[MasterSynchronizationSourceIndex]->bUseForSynchronization)
+		if (!TimeSynchronizationInputSources.IsValidIndex(MainSynchronizationSourceIndex)
+			|| TimeSynchronizationInputSources[MainSynchronizationSourceIndex] == nullptr
+			|| !TimeSynchronizationInputSources[MainSynchronizationSourceIndex]->bUseForSynchronization)
 		{
-			MasterSynchronizationSourceIndex = INDEX_NONE;
-			UE_LOG(LogTimecodeSynchronizer, Warning, TEXT("The MasterSynchronizationSourceIndex is not valid."));
+			MainSynchronizationSourceIndex = INDEX_NONE;
+			UE_LOG(LogTimecodeSynchronizer, Warning, TEXT("The MainSynchronizationSourceIndex is not valid."));
 		}
 	}
 
@@ -303,9 +303,9 @@ FFrameTime UTimecodeSynchronizer::GetProviderFrameTime() const
 
 	if (TimecodeProviderType == ETimecodeSynchronizationTimecodeType::InputSource)
 	{
-		if (SynchronizedSources.IsValidIndex(ActiveMasterSynchronizationTimecodedSourceIndex))
+		if (SynchronizedSources.IsValidIndex(ActiveMainSynchronizationTimecodedSourceIndex))
 		{
-			const FTimecodeSynchronizerActiveTimecodedInputSource& TimecodedInputSource = SynchronizedSources[ActiveMasterSynchronizationTimecodedSourceIndex];
+			const FTimecodeSynchronizerActiveTimecodedInputSource& TimecodedInputSource = SynchronizedSources[ActiveMainSynchronizationTimecodedSourceIndex];
 
 			if (GFrameCounter != LastUpdatedSources)
 			{
@@ -532,7 +532,7 @@ bool UTimecodeSynchronizer::StartSynchronization()
 	}
 	else
 	{
-		if (!ensure(SynchronizedSources.Num() == 0) || !ensure(NonSynchronizedSources.Num() == 0) || !ensure(ActiveMasterSynchronizationTimecodedSourceIndex))
+		if (!ensure(SynchronizedSources.Num() == 0) || !ensure(NonSynchronizedSources.Num() == 0) || !ensure(ActiveMainSynchronizationTimecodedSourceIndex))
 		{
 			UE_LOG(LogTimecodeSynchronizer, Error, TEXT("StartSynchronization called without properly closing sources"));
 			CloseSources();
@@ -547,9 +547,9 @@ bool UTimecodeSynchronizer::StartSynchronization()
 			UE_LOG(LogTimecodeSynchronizer, Warning, TEXT("No sources available to synchronize."));
 			SwitchState(ESynchronizationState::Error);
 		}
-		else if (TimecodeProviderType == ETimecodeSynchronizationTimecodeType::InputSource && ActiveMasterSynchronizationTimecodedSourceIndex == INDEX_NONE)
+		else if (TimecodeProviderType == ETimecodeSynchronizationTimecodeType::InputSource && ActiveMainSynchronizationTimecodedSourceIndex == INDEX_NONE)
 		{
-			UE_LOG(LogTimecodeSynchronizer, Warning, TEXT("The Master Synchronization Source could not be found."));
+			UE_LOG(LogTimecodeSynchronizer, Warning, TEXT("The Main Synchronization Source could not be found."));
 			SwitchState(ESynchronizationState::Error);
 		}
 		else
@@ -730,14 +730,14 @@ bool UTimecodeSynchronizer::Tick_TestTimecode()
 
 	if (TimecodeProviderType == ETimecodeSynchronizationTimecodeType::InputSource)
 	{
-		if (!SynchronizedSources.IsValidIndex(ActiveMasterSynchronizationTimecodedSourceIndex))
+		if (!SynchronizedSources.IsValidIndex(ActiveMainSynchronizationTimecodedSourceIndex))
 		{
-			UE_LOG(LogTimecodeSynchronizer, Error, TEXT("The InputSource '%d' that we try to synchronize on is not valid."), ActiveMasterSynchronizationTimecodedSourceIndex);
+			UE_LOG(LogTimecodeSynchronizer, Error, TEXT("The InputSource '%d' that we try to synchronize on is not valid."), ActiveMainSynchronizationTimecodedSourceIndex);
 			SwitchState(ESynchronizationState::Error);
 			return false;
 		}
 
-		return SynchronizedSources[ActiveMasterSynchronizationTimecodedSourceIndex].IsReady();
+		return SynchronizedSources[ActiveMainSynchronizationTimecodedSourceIndex].IsReady();
 	}
 	else 
 	{
@@ -926,9 +926,9 @@ void UTimecodeSynchronizer::OpenSources()
 				if (InputSource->bUseForSynchronization)
 				{
 					FTimecodeSynchronizerActiveTimecodedInputSource& NewSource = SynchronizedSources.Emplace_GetRef(InputSource);
-					if (TimecodeProviderType == ETimecodeSynchronizationTimecodeType::InputSource && Index == MasterSynchronizationSourceIndex)
+					if (TimecodeProviderType == ETimecodeSynchronizationTimecodeType::InputSource && Index == MainSynchronizationSourceIndex)
 					{
-						ActiveMasterSynchronizationTimecodedSourceIndex = SynchronizedSources.Num() - 1;
+						ActiveMainSynchronizationTimecodedSourceIndex = SynchronizedSources.Num() - 1;
 					}
 				}
 				else
@@ -943,7 +943,7 @@ void UTimecodeSynchronizer::OpenSources()
 		}
 	}
 
-	// Don't consider dynamic sources for master synchronization index.
+	// Don't consider dynamic sources for the main synchronization index.
 	for (UTimeSynchronizationSource* InputSource : DynamicSources)
 	{
 		if (InputSource != nullptr)
@@ -976,7 +976,7 @@ void UTimecodeSynchronizer::CloseSources()
 
 	SynchronizedSources.Reset();
 	NonSynchronizedSources.Reset();
-	ActiveMasterSynchronizationTimecodedSourceIndex = INDEX_NONE;
+	ActiveMainSynchronizationTimecodedSourceIndex = INDEX_NONE;
 }
 
 void UTimecodeSynchronizer::UpdateSourceStates()
diff --git a/Engine/Plugins/Media/TimecodeSynchronizer/Source/TimecodeSynchronizer/Public/TimecodeSynchronizer.h b/Engine/Plugins/Media/TimecodeSynchronizer/Source/TimecodeSynchronizer/Public/TimecodeSynchronizer.h
--- a/Engine/Plugins/Media/TimecodeSynchronizer/Source/TimecodeSynchronizer/Public/TimecodeSynchronizer.h
+++ b/Engine/Plugins/Media/TimecodeSynchronizer/Source/TimecodeSynchronizer/Public/TimecodeSynchronizer.h
@@ -312,8 +312,8 @@ public:
 	/** Returns the list of sources that are not actively being used in synchronization. */
 	const TArray<FTimecodeSynchronizerActiveTimecodedInputSource>& GetNonSynchronizedSources() const { return NonSynchronizedSources; }
 	
-	/** Returns the index of the Master Synchronization Source in the Synchronized Sources list. */
-	int32 GetActiveMasterSynchronizationTimecodedSourceIndex() const { return MasterSynchronizationSourceIndex; }
+	/** Returns the index of the Main Synchronization Source in the Synchronized Sources list. */
+	int32 GetActiveMainSynchronizationTimecodedSourceIndex() const { return MainSynchronizationSourceIndex; }
 
 	/**
 	 * Get an event delegate that is invoked when a Asset synchronization event occurred.
@@ -477,7 +477,7 @@ public:
 	 * The source need to be timecoded and flag as bUseForSynchronization
 	 */
 	UPROPERTY(EditAnywhere, Category="Timecode Provider")
-	int32 MasterSynchronizationSourceIndex;
+	int32 MainSynchronizationSourceIndex;
 
 public:
 	/** Enable verification of margin between synchronized time and source time */
@@ -575,7 +575,7 @@ private:
 	bool bPreviousUseFixedFrameRate;
 	
 	/** Index of the active source that drives the synchronized Timecode*/
-	int32 ActiveMasterSynchronizationTimecodedSourceIndex;
+	int32 ActiveMainSynchronizationTimecodedSourceIndex;
 
 	/** An event delegate that is invoked when a synchronization event occurred. */
 	FOnTimecodeSynchronizationEvent SynchronizationEvent;
diff --git a/Engine/Plugins/Media/TimecodeSynchronizer/Source/TimecodeSynchronizerEditor/Private/Widgets/STimecodeSynchronizerSourceViewport.cpp b/Engine/Plugins/Media/TimecodeSynchronizer/Source/TimecodeSynchronizerEditor/Private/Widgets/STimecodeSynchronizerSourceViewport.cpp
--- a/Engine/Plugins/Media/TimecodeSynchronizer/Source/TimecodeSynchronizerEditor/Private/Widgets/STimecodeSynchronizerSourceViewport.cpp
+++ b/Engine/Plugins/Media/TimecodeSynchronizer/Source/TimecodeSynchronizerEditor/Private/Widgets/STimecodeSynchronizerSourceViewport.cpp
@@ -131,13 +131,13 @@ void STimecodeSynchronizerSourceViewport::Construct(const FArguments& InArgs, UT
 						.HAlign(HAlign_Left)
 						.VAlign(VAlign_Bottom)
 						[
-							// Display if source is the master
+							// Display if source is the owner
 							SNew(STextBlock)
 							.ColorAndOpacity(FSlateColor::UseSubduedForeground())
 							.Font(Font18)
 							.ShadowOffset(FVector2D(1.f, 1.f))
-							.Text(this, &STimecodeSynchronizerSourceViewport::HandleIsSourceMasterText)
-							.ToolTipText(LOCTEXT("OverlayMasterSourceTooltip", "Is this source used as the master"))
+							.Text(this, &STimecodeSynchronizerSourceViewport::HandleIsSourceMainText)
+							.ToolTipText(LOCTEXT("OverlayOwnerSourceTooltip", "Is this source used as the owner"))
 						]
 						
 						+ SHorizontalBox::Slot()
@@ -194,15 +194,15 @@ FText STimecodeSynchronizerSourceViewport::HandleCurrentTimecodeText() const
 	return FText::FromString(FString("Current: ") + FApp::GetTimecode().ToString());
 }
 
-FText STimecodeSynchronizerSourceViewport::HandleIsSourceMasterText() const
+FText STimecodeSynchronizerSourceViewport::HandleIsSourceMainText() const
 {
-	FString Role;
-	if (TimecodeSynchronization && AttachedSourceIndex != INDEX_NONE && bIsSynchronizedSource && TimecodeSynchronization->GetActiveMasterSynchronizationTimecodedSourceIndex() == AttachedSourceIndex)
+	FText Role;
+	if (TimecodeSynchronization && AttachedSourceIndex != INDEX_NONE && bIsSynchronizedSource && TimecodeSynchronization->GetActiveMainSynchronizationTimecodedSourceIndex() == AttachedSourceIndex)
 	{
-		Role = "Master";
+		Role = LOCTEXT("MainLabel", "Main");
 	}
 
-	return FText::FromString(Role);
+	return Role;
 }
 
 const FTimecodeSynchronizerActiveTimecodedInputSource* STimecodeSynchronizerSourceViewport::GetAttachedSource() const
diff --git a/Engine/Plugins/Media/TimecodeSynchronizer/Source/TimecodeSynchronizerEditor/Private/Widgets/STimecodeSynchronizerSourceViewport.h b/Engine/Plugins/Media/TimecodeSynchronizer/Source/TimecodeSynchronizerEditor/Private/Widgets/STimecodeSynchronizerSourceViewport.h
--- a/Engine/Plugins/Media/TimecodeSynchronizer/Source/TimecodeSynchronizerEditor/Private/Widgets/STimecodeSynchronizerSourceViewport.h
+++ b/Engine/Plugins/Media/TimecodeSynchronizer/Source/TimecodeSynchronizerEditor/Private/Widgets/STimecodeSynchronizerSourceViewport.h
@@ -50,8 +50,8 @@ private:
 	/** Callback for getting the text of the Current Synchronized Timecode. */
 	FText HandleCurrentTimecodeText() const;
 	
-	/** Callback setuping the text if the source is the master source. */
-	FText HandleIsSourceMasterText() const;
+	/** Callback setuping the text if the source is the main source. */
+	FText HandleIsSourceMainText() const;
 
 	/** Get the attached Input source in TimecodeSynchronization.GetActiveSources() at the AttachedSourceIndex index  */
 	const FTimecodeSynchronizerActiveTimecodedInputSource* GetAttachedSource() const;
diff --git a/Engine/Plugins/Media/TimecodeSynchronizer/Source/TimecodeSynchronizerEditor/Private/Widgets/STimecodeSynchronizerWidget.cpp b/Engine/Plugins/Media/TimecodeSynchronizer/Source/TimecodeSynchronizerEditor/Private/Widgets/STimecodeSynchronizerWidget.cpp
--- a/Engine/Plugins/Media/TimecodeSynchronizer/Source/TimecodeSynchronizerEditor/Private/Widgets/STimecodeSynchronizerWidget.cpp
+++ b/Engine/Plugins/Media/TimecodeSynchronizer/Source/TimecodeSynchronizerEditor/Private/Widgets/STimecodeSynchronizerWidget.cpp
@@ -29,7 +29,7 @@ STimecodeSynchronizerBarWidget::STimecodeSynchronizerBarWidget()
 	, CurrentFrameWidth(10)
 	, CurrentFrameValue(0)
 	, CurrentFrame(0)
-	, CurrentMasterIndex(0)
+	, CurrentOwnerIndex(0)
 	, MinOldestFrameTime(0)
 	, MaxOldestFrameTime(0)
 	, MinNewestFrameTime(0)
@@ -105,7 +105,7 @@ void STimecodeSynchronizerBarWidget::Tick(const FGeometry& InAllottedGeometry, c
 		CurrentFrameValue = TimecodeSynchronizer->GetCurrentSystemFrameTime().GetFrame().Value;
 		CurrentFrameRate = TimecodeSynchronizer->GetFrameRate();
 		CurrentFrame = TimecodeSynchronizer->GetCurrentSystemFrameTime().GetFrame().Value * CurrentFrameWidth;
-		CurrentMasterIndex = TimecodeSynchronizer->GetActiveMasterSynchronizationTimecodedSourceIndex();
+		CurrentOwnerIndex = TimecodeSynchronizer->GetActiveMainSynchronizationTimecodedSourceIndex();
 	}
 }
 
@@ -241,7 +241,7 @@ int32 STimecodeSynchronizerBarWidget::OnPaint(const FPaintArgs& InArgs, const FG
 					OutDrawElements,
 					InLayerId,
 					LeftJustification,
-					FString::Printf(TEXT("[%d] - %s %s (%d)"), i, *DisplayData[i].Name, (i == CurrentMasterIndex) ? TEXT(" - Master") : TEXT(""), (DisplayData[i].NewestFrameTime - DisplayData[i].OldestFrameTime) / CurrentFrameWidth),
+					FString::Printf(TEXT("[%d] - %s %s (%d)"), i, *DisplayData[i].Name, (i == CurrentOwnerIndex) ? TEXT(" - Owner") : TEXT(""), (DisplayData[i].NewestFrameTime - DisplayData[i].OldestFrameTime) / CurrentFrameWidth),
 					FBox2D(FVector2D(TextMinX, MinY), FVector2D(TextMaxX, MaxY)),
 					InWidgetStyle);
 			}
diff --git a/Engine/Plugins/Media/TimecodeSynchronizer/Source/TimecodeSynchronizerEditor/Private/Widgets/STimecodeSynchronizerWidget.h b/Engine/Plugins/Media/TimecodeSynchronizer/Source/TimecodeSynchronizerEditor/Private/Widgets/STimecodeSynchronizerWidget.h
--- a/Engine/Plugins/Media/TimecodeSynchronizer/Source/TimecodeSynchronizerEditor/Private/Widgets/STimecodeSynchronizerWidget.h
+++ b/Engine/Plugins/Media/TimecodeSynchronizer/Source/TimecodeSynchronizerEditor/Private/Widgets/STimecodeSynchronizerWidget.h
@@ -60,7 +60,7 @@ private:
 	FFrameRate CurrentFrameRate;
 	int32 CurrentFrameValue;
 	int32 CurrentFrame;
-	int32 CurrentMasterIndex;
+	int32 CurrentOwnerIndex;
 
 	struct FrameTimeDisplayData
 	{
